- Build fix for Mac

diff --git a/Engine/Source/Runtime/NetworkFile/Private/ITransport.h b/Engine/Source/Runtime/NetworkFile/Private/ITransport.h
--- a/Engine/Source/Runtime/NetworkFile/Private/ITransport.h
+++ b/Engine/Source/Runtime/NetworkFile/Private/ITransport.h
@@ -12,7 +12,7 @@ class ITransport
 
 public:
 	ITransport (){};
-	~ITransport(){};
+	virtual ~ITransport(){};
 
 	virtual bool Initialize(const TCHAR*) = 0; 
 	virtual bool SendPayloadAndReceiveResponse(TArray<uint8>& In, TArray<uint8>& Out) = 0; 
