Fix for Non-unity build, missing include from CL# 2103382

diff --git a/Engine/Source/Runtime/NetworkFile/Private/NetworkFilePrivatePCH.h b/Engine/Source/Runtime/NetworkFile/Private/NetworkFilePrivatePCH.h
--- a/Engine/Source/Runtime/NetworkFile/Private/NetworkFilePrivatePCH.h
+++ b/Engine/Source/Runtime/NetworkFile/Private/NetworkFilePrivatePCH.h
@@ -2,6 +2,7 @@
 
 #pragma once
 
-#include "Core.h"
-#include "ModuleManager.h"
+#include "NetworkMessage.h"
 
+#include "Core.h"
+#include "ModuleManager.h"
\ No newline at end of file
diff --git a/Engine/Source/Runtime/NetworkFile/Private/TCPTransport.cpp b/Engine/Source/Runtime/NetworkFile/Private/TCPTransport.cpp
--- a/Engine/Source/Runtime/NetworkFile/Private/TCPTransport.cpp
+++ b/Engine/Source/Runtime/NetworkFile/Private/TCPTransport.cpp
@@ -2,7 +2,9 @@
 
 #include "NetworkFilePrivatePCH.h"
 #include "TCPTransport.h"
-#include "Sockets.h"
+#include "MultichannelTCP.h"
+
+DECLARE_LOG_CATEGORY_EXTERN(LogNetworkPlatformFile, Log, All);
 
 FTCPTransport::FTCPTransport()
 	:FileSocket(NULL)
