UE4 Refactoring. Changed OVERRIDE and FINAL macros to keywords override and final.

diff --git a/Engine/Source/Runtime/GameplayAbilities/Public/GameplayEffectExtension_ShieldTest.h b/Engine/Source/Runtime/GameplayAbilities/Public/GameplayEffectExtension_ShieldTest.h
--- a/Engine/Source/Runtime/GameplayAbilities/Public/GameplayEffectExtension_ShieldTest.h
+++ b/Engine/Source/Runtime/GameplayAbilities/Public/GameplayEffectExtension_ShieldTest.h
@@ -15,8 +15,8 @@ class GAMEPLAYABILITIES_API UGameplayEffectExtension_ShieldTest : public UGamepl
 
 public:
 
-	void PreGameplayEffectExecute(const FGameplayModifierEvaluatedData &SelfData, FGameplayEffectModCallbackData &Data) const OVERRIDE;
-	void PostGameplayEffectExecute(const FGameplayModifierEvaluatedData &SelfData, const FGameplayEffectModCallbackData &Data) const OVERRIDE;
+	void PreGameplayEffectExecute(const FGameplayModifierEvaluatedData &SelfData, FGameplayEffectModCallbackData &Data) const override;
+	void PostGameplayEffectExecute(const FGameplayModifierEvaluatedData &SelfData, const FGameplayEffectModCallbackData &Data) const override;
 
 	UPROPERTY(EditDefaultsOnly, Category = Lifesteal)
 	UGameplayEffect * ShieldRemoveGameplayEffect;
diff --git a/Engine/Source/Runtime/GameplayAbilities/Public/GameplayEffectStackingExtension_CappedNumberTest.h b/Engine/Source/Runtime/GameplayAbilities/Public/GameplayEffectStackingExtension_CappedNumberTest.h
--- a/Engine/Source/Runtime/GameplayAbilities/Public/GameplayEffectStackingExtension_CappedNumberTest.h
+++ b/Engine/Source/Runtime/GameplayAbilities/Public/GameplayEffectStackingExtension_CappedNumberTest.h
@@ -15,5 +15,5 @@ class GAMEPLAYABILITIES_API UGameplayEffectStackingExtension_CappedNumberTest :
 
 public:
 
-	void CalculateStack(TArray<FActiveGameplayEffect*>& CustomGameplayEffects, FActiveGameplayEffectsContainer& Container, FActiveGameplayEffect& CurrentEffect) OVERRIDE;
+	void CalculateStack(TArray<FActiveGameplayEffect*>& CustomGameplayEffects, FActiveGameplayEffectsContainer& Container, FActiveGameplayEffect& CurrentEffect) override;
 };
\ No newline at end of file
diff --git a/Engine/Source/Runtime/GameplayAbilities/Public/GameplayEffectStackingExtension_DiminishingReturnsTest.h b/Engine/Source/Runtime/GameplayAbilities/Public/GameplayEffectStackingExtension_DiminishingReturnsTest.h
--- a/Engine/Source/Runtime/GameplayAbilities/Public/GameplayEffectStackingExtension_DiminishingReturnsTest.h
+++ b/Engine/Source/Runtime/GameplayAbilities/Public/GameplayEffectStackingExtension_DiminishingReturnsTest.h
@@ -15,5 +15,5 @@ class GAMEPLAYABILITIES_API UGameplayEffectStackingExtension_DiminishingReturnsT
 
 public:
 
-	void CalculateStack(TArray<FActiveGameplayEffect*>& CustomGameplayEffects, FActiveGameplayEffectsContainer& Container, FActiveGameplayEffect& CurrentEffect) OVERRIDE;
+	void CalculateStack(TArray<FActiveGameplayEffect*>& CustomGameplayEffects, FActiveGameplayEffectsContainer& Container, FActiveGameplayEffect& CurrentEffect) override;
 };
\ No newline at end of file
diff --git a/Engine/Source/Runtime/GameplayTags/Private/GameplayTagsModule.cpp b/Engine/Source/Runtime/GameplayTags/Private/GameplayTagsModule.cpp
--- a/Engine/Source/Runtime/GameplayTags/Private/GameplayTagsModule.cpp
+++ b/Engine/Source/Runtime/GameplayTags/Private/GameplayTagsModule.cpp
@@ -5,12 +5,12 @@
 class FGameplayTagsModule : public IGameplayTagsModule
 {
 	// Begin IModuleInterface
-	virtual void StartupModule() OVERRIDE;
-	virtual void ShutdownModule() OVERRIDE;
+	virtual void StartupModule() override;
+	virtual void ShutdownModule() override;
 	// End IModuleInterface
 
 	// Gets the UGameplayTagsManager manager
-	virtual UGameplayTagsManager& GetGameplayTagsManager() OVERRIDE
+	virtual UGameplayTagsManager& GetGameplayTagsManager() override
 	{
 		check(GameplayTagsManager);
 		return *GameplayTagsManager;
diff --git a/Engine/Source/Runtime/HTML5/HTML5Audio/Private/HTML5AudioDevice.cpp b/Engine/Source/Runtime/HTML5/HTML5Audio/Private/HTML5AudioDevice.cpp
--- a/Engine/Source/Runtime/HTML5/HTML5Audio/Private/HTML5AudioDevice.cpp
+++ b/Engine/Source/Runtime/HTML5/HTML5Audio/Private/HTML5AudioDevice.cpp
@@ -30,7 +30,7 @@ class FHTML5AudioDeviceModule : public IAudioDeviceModule
 public:
 
 	/** Creates a new instance of the audio device implemented by the module. */
-	virtual FAudioDevice* CreateAudioDevice() OVERRIDE
+	virtual FAudioDevice* CreateAudioDevice() override
 	{
 		return new FALAudioDevice;
 	}
diff --git a/Engine/Source/Runtime/HTML5/HTML5Audio/Public/HTML5AudioDevice.h b/Engine/Source/Runtime/HTML5/HTML5Audio/Public/HTML5AudioDevice.h
--- a/Engine/Source/Runtime/HTML5/HTML5Audio/Public/HTML5AudioDevice.h
+++ b/Engine/Source/Runtime/HTML5/HTML5Audio/Public/HTML5AudioDevice.h
@@ -197,20 +197,20 @@ public:
 	FALAudioDevice() {} 
 	virtual ~FALAudioDevice() {} 
 
-	virtual FName GetRuntimeFormat(USoundWave* SoundWave) OVERRIDE
+	virtual FName GetRuntimeFormat(USoundWave* SoundWave) override
 	{
 		static FName NAME_OGG(TEXT("OGG"));
 		return NAME_OGG;
 	}
 
 	/** Starts up any platform specific hardware/APIs */
-	virtual bool InitializeHardware() OVERRIDE;
+	virtual bool InitializeHardware() override;
  
 	/**
 	 * Tears down audio device by stopping all sounds, removing all buffers,  destroying all sources, ... Called by both Destroy and ShutdownAfterError
 	 * to perform the actual tear down.
 	 */
-	virtual void TeardownHardware() OVERRIDE;
+	virtual void TeardownHardware() override;
 
 	/**
 	 * Update the audio device and calculates the cached inverse transform later
@@ -232,9 +232,9 @@ public:
 	 */
 	virtual void FreeResource( USoundWave* SoundWave );
 
-	virtual bool HasCompressedAudioInfoClass(USoundWave* SoundWave) OVERRIDE;
+	virtual bool HasCompressedAudioInfoClass(USoundWave* SoundWave) override;
 
-	virtual class ICompressedAudioInfo* CreateCompressedAudioInfo(USoundWave* SoundWave) OVERRIDE;
+	virtual class ICompressedAudioInfo* CreateCompressedAudioInfo(USoundWave* SoundWave) override;
 
 	void FindProcs( bool AllowExt );
 
@@ -243,7 +243,7 @@ public:
 
 protected:
 
-	virtual FSoundSource* CreateSoundSource() OVERRIDE; 
+	virtual FSoundSource* CreateSoundSource() override; 
 	/** Returns the enum for the internal format for playing a sound with this number of channels. */
 	ALuint GetInternalFormat( int NumChannels );
 
diff --git a/Engine/Source/Runtime/HeadMountedDisplay/Public/IHeadMountedDisplayModule.h b/Engine/Source/Runtime/HeadMountedDisplay/Public/IHeadMountedDisplayModule.h
--- a/Engine/Source/Runtime/HeadMountedDisplay/Public/IHeadMountedDisplayModule.h
+++ b/Engine/Source/Runtime/HeadMountedDisplay/Public/IHeadMountedDisplayModule.h
@@ -19,7 +19,7 @@ public:
 		return HMDFeatureName;
 	}
 
-	virtual void StartupModule() OVERRIDE
+	virtual void StartupModule() override
 	{
 		IModularFeatures::Get().RegisterModularFeature( GetModularFeatureName(), this );
 	}
diff --git a/Engine/Source/Runtime/IOS/IOSAudio/Public/IOSAudioDevice.h b/Engine/Source/Runtime/IOS/IOSAudio/Public/IOSAudioDevice.h
--- a/Engine/Source/Runtime/IOS/IOSAudio/Public/IOSAudioDevice.h
+++ b/Engine/Source/Runtime/IOS/IOSAudio/Public/IOSAudioDevice.h
@@ -79,7 +79,7 @@ public:
 	 *
 	 * @return Size in bytes
 	 */
-	virtual int32 GetSize(void) OVERRIDE;
+	virtual int32 GetSize(void) override;
 
 	/** Format of the sound referenced by this buffer */
 	int32  SoundFormat;
@@ -115,22 +115,22 @@ public:
 	 * @param	WaveInstance	wave instance being primed for playback
 	 * @return	true			if initialization was successful, false otherwise
 	 */
-	virtual bool Init(FWaveInstance* WaveInstance) OVERRIDE;
+	virtual bool Init(FWaveInstance* WaveInstance) override;
 	
 	/**
 	 * Updates the source specific parameter like e.g. volume and pitch based on the associated
 	 * wave instance.	
 	 */
-	virtual void Update(void) OVERRIDE;
+	virtual void Update(void) override;
 	
 	/**  Plays the current wave instance. */
-	virtual void Play(void) OVERRIDE;
+	virtual void Play(void) override;
 	
 	/** Stops the current wave instance and detaches it from the source. */
-	virtual void Stop(void) OVERRIDE;
+	virtual void Stop(void) override;
 	
 	/** Pauses playback of current wave instance. */
-	virtual void Pause(void) OVERRIDE;
+	virtual void Pause(void) override;
 	
 	/**
 	 * Queries the status of the currently associated wave instance.
@@ -138,7 +138,7 @@ public:
 	 * @return	true if the wave instance/ source has finished playback and false if it is 
 	 *			currently playing or paused.
 	 */
-	virtual bool IsFinished(void) OVERRIDE;
+	virtual bool IsFinished(void) override;
 
 	/** Calculates the audio unit element of the input channel relative to the base bus number */
 	AudioUnitElement GetAudioUnitElement(int32 Channel);
@@ -178,7 +178,7 @@ public:
 	FIOSAudioDevice();
 	virtual ~FIOSAudioDevice() { }
 	
-	virtual FName GetRuntimeFormat(USoundWave* SoundWave) OVERRIDE
+	virtual FName GetRuntimeFormat(USoundWave* SoundWave) override
 	{
 		static FName NAME_ADPCM(TEXT("ADPCM"));
 		return NAME_ADPCM;
@@ -192,23 +192,23 @@ public:
 	virtual void ResumeContext();
 	virtual void SuspendContext();
 	
-	virtual bool Exec( UWorld* InWorld, const TCHAR* Cmd, FOutputDevice& Ar = *GLog ) OVERRIDE;
+	virtual bool Exec( UWorld* InWorld, const TCHAR* Cmd, FOutputDevice& Ar = *GLog ) override;
 
 protected:
 	/** Starts up any platform specific hardware/APIs */
-	virtual bool InitializeHardware() OVERRIDE;
+	virtual bool InitializeHardware() override;
 
 	/** Shuts down any platform specific hardware/APIs */
-	virtual void TeardownHardware() OVERRIDE;
+	virtual void TeardownHardware() override;
 
 	/** Lets the platform any tick actions */
-	virtual void UpdateHardware() OVERRIDE;
+	virtual void UpdateHardware() override;
 
 	/** Creates a new platform specific sound source */
-	virtual FAudioEffectsManager* CreateEffectsManager() OVERRIDE;
+	virtual FAudioEffectsManager* CreateEffectsManager() override;
 
 	/** Creates a new platform specific sound source */
-	virtual FSoundSource* CreateSoundSource() OVERRIDE;
+	virtual FSoundSource* CreateSoundSource() override;
 
 	/** Audio Session management */
 	void GetHardwareSampleRate(double& OutSampleRate);
diff --git a/Engine/Source/Runtime/IOS/MetalRHI/Private/MetalRHIPrivate.h b/Engine/Source/Runtime/IOS/MetalRHI/Private/MetalRHIPrivate.h
--- a/Engine/Source/Runtime/IOS/MetalRHI/Private/MetalRHIPrivate.h
+++ b/Engine/Source/Runtime/IOS/MetalRHI/Private/MetalRHIPrivate.h
@@ -139,7 +139,7 @@ public:
 		}
 	}
 	
-	virtual bool Create(bool _bIsManualReset = false) OVERRIDE
+	virtual bool Create(bool _bIsManualReset = false) override
 	{
 		check(!bInitialized);
 		bool RetVal = false;
@@ -161,7 +161,7 @@ public:
 		return RetVal;
 	}
 	
-	virtual void Trigger() OVERRIDE
+	virtual void Trigger() override
 	{
 		check(bInitialized);
 		
@@ -186,7 +186,7 @@ public:
 		UnlockEventMutex();
 	}
 	
-	virtual void Reset() OVERRIDE
+	virtual void Reset() override
 	{
 		check(bInitialized);
 		LockEventMutex();
diff --git a/Engine/Source/Runtime/InputCore/Public/InputCoreModule.h b/Engine/Source/Runtime/InputCore/Public/InputCoreModule.h
--- a/Engine/Source/Runtime/InputCore/Public/InputCoreModule.h
+++ b/Engine/Source/Runtime/InputCore/Public/InputCoreModule.h
@@ -15,5 +15,5 @@ public:
 	 * Called right after the module DLL has been loaded and the module object has been created
 	 * Overloaded to allow the default subsystem a chance to load
 	 */
-	virtual void StartupModule() OVERRIDE;
+	virtual void StartupModule() override;
 };
\ No newline at end of file
diff --git a/Engine/Source/Runtime/Launch/Private/LaunchEngineLoop.cpp b/Engine/Source/Runtime/Launch/Private/LaunchEngineLoop.cpp
--- a/Engine/Source/Runtime/Launch/Private/LaunchEngineLoop.cpp
+++ b/Engine/Source/Runtime/Launch/Private/LaunchEngineLoop.cpp
@@ -98,7 +98,7 @@ public:
 	{
 	}
 
-	virtual void Serialize( const TCHAR* V, ELogVerbosity::Type Verbosity, const class FName& Category ) OVERRIDE
+	virtual void Serialize( const TCHAR* V, ELogVerbosity::Type Verbosity, const class FName& Category ) override
 	{
 		if ( (bAllowLogVerbosity && Verbosity <= ELogVerbosity::Log) || (Verbosity <= ELogVerbosity::Display) )
 		{
@@ -127,7 +127,7 @@ public:
 	virtual ~FOutputDeviceTestExit()
 	{
 	}
-	virtual void Serialize(const TCHAR* V, ELogVerbosity::Type Verbosity, const class FName& Category) OVERRIDE
+	virtual void Serialize(const TCHAR* V, ELogVerbosity::Type Verbosity, const class FName& Category) override
 	{
 		if (!GIsRequestingExit)
 		{
diff --git a/Engine/Source/Runtime/Launch/Public/LaunchEngineLoop.h b/Engine/Source/Runtime/Launch/Public/LaunchEngineLoop.h
--- a/Engine/Source/Runtime/Launch/Public/LaunchEngineLoop.h
+++ b/Engine/Source/Runtime/Launch/Public/LaunchEngineLoop.h
@@ -66,7 +66,7 @@ public:
 	 *
 	 * @return The error level; 0 if successful, > 0 if there were errors.
 	 */ 
-	virtual int32 Init() OVERRIDE;
+	virtual int32 Init() override;
 
 	/**
 	 * Initialize the timing options from the command line.
@@ -84,10 +84,10 @@ public:
 	/**
 	 * Advances the main loop.
 	 */
-	virtual void Tick() OVERRIDE;
+	virtual void Tick() override;
 
 	/** Removes references to any objects pending cleanup by deleting them. */
-	virtual void ClearPendingCleanupObjects() OVERRIDE;
+	virtual void ClearPendingCleanupObjects() override;
 
 #if PLATFORM_XBOXONE
 	void OnResuming(_In_ Platform::Object^ Sender, _In_ Platform::Object^ Args);
diff --git a/Engine/Source/Runtime/LauncherAutomatedService/Private/AutomatedTestManager.h b/Engine/Source/Runtime/LauncherAutomatedService/Private/AutomatedTestManager.h
--- a/Engine/Source/Runtime/LauncherAutomatedService/Private/AutomatedTestManager.h
+++ b/Engine/Source/Runtime/LauncherAutomatedService/Private/AutomatedTestManager.h
@@ -42,9 +42,9 @@ public:
 
 	// IAutomatedTestManager interface
 
-	virtual bool IsTestingComplete() OVERRIDE;
-	virtual void RunTests() OVERRIDE;
-	virtual void Tick( float DeltaTime ) OVERRIDE;
+	virtual bool IsTestingComplete() override;
+	virtual void RunTests() override;
+	virtual void Tick( float DeltaTime ) override;
 
 protected:
 
diff --git a/Engine/Source/Runtime/LauncherAutomatedService/Private/LauncherAutomatedServiceModule.h b/Engine/Source/Runtime/LauncherAutomatedService/Private/LauncherAutomatedServiceModule.h
--- a/Engine/Source/Runtime/LauncherAutomatedService/Private/LauncherAutomatedServiceModule.h
+++ b/Engine/Source/Runtime/LauncherAutomatedService/Private/LauncherAutomatedServiceModule.h
@@ -13,7 +13,7 @@ public:
 
 	// ILauncherAutomatedServiceModule interface
 
-	virtual ILauncherAutomatedServiceProviderPtr CreateAutomatedServiceProvider() OVERRIDE;
+	virtual ILauncherAutomatedServiceProviderPtr CreateAutomatedServiceProvider() override;
 
 public:
 
diff --git a/Engine/Source/Runtime/LauncherAutomatedService/Private/LauncherAutomatedServiceProvider.h b/Engine/Source/Runtime/LauncherAutomatedService/Private/LauncherAutomatedServiceProvider.h
--- a/Engine/Source/Runtime/LauncherAutomatedService/Private/LauncherAutomatedServiceProvider.h
+++ b/Engine/Source/Runtime/LauncherAutomatedService/Private/LauncherAutomatedServiceProvider.h
@@ -20,19 +20,19 @@ public:
 
 	// ILauncherAutomatedServiceProvider interface
 
-	virtual int32 GetExitCode( ) OVERRIDE
+	virtual int32 GetExitCode( ) override
 	{
 		return bHasErrors ? 1 : 0;
 	}
 
-	virtual bool IsRunning( ) OVERRIDE
+	virtual bool IsRunning( ) override
 	{
 		return bIsReadyToShutdown == false && bHasErrors == false;
 	}
 
-	virtual void Setup( const TCHAR* Params ) OVERRIDE;
-	virtual void Shutdown( ) OVERRIDE;
-	virtual void Tick( float DeltaTime ) OVERRIDE;
+	virtual void Setup( const TCHAR* Params ) override;
+	virtual void Shutdown( ) override;
+	virtual void Tick( float DeltaTime ) override;
 
 protected:
 
diff --git a/Engine/Source/Runtime/Linux/ALAudio/Public/ALAudioDevice.h b/Engine/Source/Runtime/Linux/ALAudio/Public/ALAudioDevice.h
--- a/Engine/Source/Runtime/Linux/ALAudio/Public/ALAudioDevice.h
+++ b/Engine/Source/Runtime/Linux/ALAudio/Public/ALAudioDevice.h
@@ -195,20 +195,20 @@ public:
 	FALAudioDevice() {} 
 	virtual ~FALAudioDevice() {} 
 
-	virtual FName GetRuntimeFormat(USoundWave* SoundWave) OVERRIDE
+	virtual FName GetRuntimeFormat(USoundWave* SoundWave) override
 	{
 		static FName NAME_OGG(TEXT("OGG"));
 		return NAME_OGG;
 	}
 
 	/** Starts up any platform specific hardware/APIs */
-	virtual bool InitializeHardware() OVERRIDE;
+	virtual bool InitializeHardware() override;
  
 	/**
 	 * Tears down audio device by stopping all sounds, removing all buffers,  destroying all sources, ... Called by both Destroy and ShutdownAfterError
 	 * to perform the actual tear down.
 	 */
-	virtual void TeardownHardware() OVERRIDE;
+	virtual void TeardownHardware() override;
 
 	/**
 	 * Update the audio device and calculates the cached inverse transform later
@@ -230,9 +230,9 @@ public:
 	 */
 	virtual void FreeResource( USoundWave* SoundWave );
 
-	virtual bool HasCompressedAudioInfoClass(USoundWave* SoundWave) OVERRIDE;
+	virtual bool HasCompressedAudioInfoClass(USoundWave* SoundWave) override;
 
-	virtual class ICompressedAudioInfo* CreateCompressedAudioInfo(USoundWave* SoundWave) OVERRIDE;
+	virtual class ICompressedAudioInfo* CreateCompressedAudioInfo(USoundWave* SoundWave) override;
 
 	void FindProcs( bool AllowExt );
 
@@ -241,7 +241,7 @@ public:
 
 protected:
 
-	virtual FSoundSource* CreateSoundSource() OVERRIDE; 
+	virtual FSoundSource* CreateSoundSource() override; 
 	/** Returns the enum for the internal format for playing a sound with this number of channels. */
 	ALuint GetInternalFormat( int NumChannels );
 
diff --git a/Engine/Source/Runtime/Mac/CoreAudio/Public/CoreAudioDevice.h b/Engine/Source/Runtime/Mac/CoreAudio/Public/CoreAudioDevice.h
--- a/Engine/Source/Runtime/Mac/CoreAudio/Public/CoreAudioDevice.h
+++ b/Engine/Source/Runtime/Mac/CoreAudio/Public/CoreAudioDevice.h
@@ -316,19 +316,19 @@ private:
 class FCoreAudioDevice : public FAudioDevice
 {
 	/** Starts up any platform specific hardware/APIs */
-	virtual bool InitializeHardware() OVERRIDE;
+	virtual bool InitializeHardware() override;
 
 	/** Shuts down any platform specific hardware/APIs */
-	virtual void TeardownHardware() OVERRIDE;
+	virtual void TeardownHardware() override;
 
 	/** Lets the platform any tick actions */
-	virtual void UpdateHardware() OVERRIDE;
+	virtual void UpdateHardware() override;
 
 	/** Creates a new platform specific sound source */
-	virtual FAudioEffectsManager* CreateEffectsManager() OVERRIDE;
+	virtual FAudioEffectsManager* CreateEffectsManager() override;
 
 	/** Creates a new platform specific sound source */
-	virtual FSoundSource* CreateSoundSource() OVERRIDE;
+	virtual FSoundSource* CreateSoundSource() override;
 
 	/** Setup Matrix Mixer's input - enable input, input->output volumes */
 	void SetupMatrixMixerInput( int32 Input, bool bIs6ChannelOGG );
@@ -354,20 +354,20 @@ class FCoreAudioDevice : public FAudioDevice
 	AUNode GetMatrixMixerNode() const { return MatrixMixerNode; }
 	AudioUnit GetMatrixMixerUnit() const { return MatrixMixerUnit; }
 	
-	virtual FName GetRuntimeFormat(USoundWave* SoundWave) OVERRIDE
+	virtual FName GetRuntimeFormat(USoundWave* SoundWave) override
 	{
 		static FName NAME_OGG(TEXT("OGG"));
 		return NAME_OGG;
 	}
 
-	virtual bool HasCompressedAudioInfoClass(USoundWave* SoundWave) OVERRIDE;
+	virtual bool HasCompressedAudioInfoClass(USoundWave* SoundWave) override;
 
-	virtual bool SupportsRealtimeDecompression() const OVERRIDE
+	virtual bool SupportsRealtimeDecompression() const override
 	{ 
 		return true; 
 	}
 
-	virtual class ICompressedAudioInfo* CreateCompressedAudioInfo(USoundWave* SoundWave) OVERRIDE;
+	virtual class ICompressedAudioInfo* CreateCompressedAudioInfo(USoundWave* SoundWave) override;
 
 	FORCEINLINE uint32 CalculateMatrixElement( uint32 InputNum, uint32 OutputNum )
 	{
diff --git a/Engine/Source/Runtime/MoviePlayer/Private/DefaultGameMoviePlayer.h b/Engine/Source/Runtime/MoviePlayer/Private/DefaultGameMoviePlayer.h
--- a/Engine/Source/Runtime/MoviePlayer/Private/DefaultGameMoviePlayer.h
+++ b/Engine/Source/Runtime/MoviePlayer/Private/DefaultGameMoviePlayer.h
@@ -14,22 +14,22 @@ public:
 	~FDefaultGameMoviePlayer();
 
 	/** IGameMoviePlayer Interface */
-	virtual void RegisterMovieStreamer(TSharedPtr<IMovieStreamer> InMovieStreamer) OVERRIDE;
-	virtual void SetSlateRenderer(TSharedPtr<FSlateRenderer> InSlateRenderer) OVERRIDE;
-	virtual void Initialize() OVERRIDE;
-	virtual void PassLoadingScreenWindowBackToGame() const OVERRIDE;
-	virtual void SetupLoadingScreen(const FLoadingScreenAttributes& LoadingScreenAttributes) OVERRIDE;
-	virtual bool PlayMovie() OVERRIDE;
-	virtual void WaitForMovieToFinish() OVERRIDE;
-	virtual bool IsLoadingFinished() const OVERRIDE;
-	virtual bool IsMovieCurrentlyPlaying() const OVERRIDE;
-	virtual bool LoadingScreenIsPrepared() const OVERRIDE;
-	virtual void SetupLoadingScreenFromIni() OVERRIDE;
+	virtual void RegisterMovieStreamer(TSharedPtr<IMovieStreamer> InMovieStreamer) override;
+	virtual void SetSlateRenderer(TSharedPtr<FSlateRenderer> InSlateRenderer) override;
+	virtual void Initialize() override;
+	virtual void PassLoadingScreenWindowBackToGame() const override;
+	virtual void SetupLoadingScreen(const FLoadingScreenAttributes& LoadingScreenAttributes) override;
+	virtual bool PlayMovie() override;
+	virtual void WaitForMovieToFinish() override;
+	virtual bool IsLoadingFinished() const override;
+	virtual bool IsMovieCurrentlyPlaying() const override;
+	virtual bool LoadingScreenIsPrepared() const override;
+	virtual void SetupLoadingScreenFromIni() override;
 
 	/** FTickableObjectRenderThread interface */
-	virtual void Tick( float DeltaTime ) OVERRIDE;
-	virtual TStatId GetStatId() const OVERRIDE;
-	virtual bool IsTickable() const OVERRIDE;
+	virtual void Tick( float DeltaTime ) override;
+	virtual TStatId GetStatId() const override;
+	virtual bool IsTickable() const override;
 	
 private:
 
diff --git a/Engine/Source/Runtime/MoviePlayer/Private/MoviePlayerThreading.cpp b/Engine/Source/Runtime/MoviePlayer/Private/MoviePlayerThreading.cpp
--- a/Engine/Source/Runtime/MoviePlayer/Private/MoviePlayerThreading.cpp
+++ b/Engine/Source/Runtime/MoviePlayer/Private/MoviePlayerThreading.cpp
@@ -27,8 +27,8 @@ public:
 	}
 
 	/** FRunnable interface */
-	virtual bool Init() OVERRIDE;
-	virtual uint32 Run() OVERRIDE;
+	virtual bool Init() override;
+	virtual uint32 Run() override;
 private:
 	/** Hold a handle to our parent sync mechanism which handles all of our threading locks */
 	class FSlateLoadingSynchronizationMechanism* SyncMechanism;
diff --git a/Engine/Source/Runtime/MoviePlayer/Private/NullMoviePlayer.h b/Engine/Source/Runtime/MoviePlayer/Private/NullMoviePlayer.h
--- a/Engine/Source/Runtime/MoviePlayer/Private/NullMoviePlayer.h
+++ b/Engine/Source/Runtime/MoviePlayer/Private/NullMoviePlayer.h
@@ -18,17 +18,17 @@ public:
 	}
 
 	/** IGameMoviePlayer Interface */
-	virtual void RegisterMovieStreamer(TSharedPtr<IMovieStreamer> InMovieStreamer) OVERRIDE {}
-	virtual void SetSlateRenderer(TSharedPtr<FSlateRenderer> InSlateRenderer) OVERRIDE {}
-	virtual void Initialize() OVERRIDE {}
-	virtual void PassLoadingScreenWindowBackToGame() const OVERRIDE {}
-	virtual void SetupLoadingScreen(const FLoadingScreenAttributes& InLoadingScreenAttributes) OVERRIDE {}
-	virtual bool PlayMovie() OVERRIDE { return false;}
-	virtual void WaitForMovieToFinish() OVERRIDE {}
-	virtual bool IsLoadingFinished() const OVERRIDE {return true;}
-	virtual bool IsMovieCurrentlyPlaying() const OVERRIDE  {return false;}
-	virtual bool LoadingScreenIsPrepared() const OVERRIDE {return false;}
-	virtual void SetupLoadingScreenFromIni() OVERRIDE {}
+	virtual void RegisterMovieStreamer(TSharedPtr<IMovieStreamer> InMovieStreamer) override {}
+	virtual void SetSlateRenderer(TSharedPtr<FSlateRenderer> InSlateRenderer) override {}
+	virtual void Initialize() override {}
+	virtual void PassLoadingScreenWindowBackToGame() const override {}
+	virtual void SetupLoadingScreen(const FLoadingScreenAttributes& InLoadingScreenAttributes) override {}
+	virtual bool PlayMovie() override { return false;}
+	virtual void WaitForMovieToFinish() override {}
+	virtual bool IsLoadingFinished() const override {return true;}
+	virtual bool IsMovieCurrentlyPlaying() const override  {return false;}
+	virtual bool LoadingScreenIsPrepared() const override {return false;}
+	virtual void SetupLoadingScreenFromIni() override {}
 
 private:
 	FNullGameMoviePlayer() {}
diff --git a/Engine/Source/Runtime/MoviePlayer/Public/MoviePlayer.h b/Engine/Source/Runtime/MoviePlayer/Public/MoviePlayer.h
--- a/Engine/Source/Runtime/MoviePlayer/Public/MoviePlayer.h
+++ b/Engine/Source/Runtime/MoviePlayer/Public/MoviePlayer.h
@@ -15,17 +15,17 @@ public:
 	~FMovieViewport() {}
 
 	/* ISlateViewport interface */
-	virtual FIntPoint GetSize() const OVERRIDE
+	virtual FIntPoint GetSize() const override
 	{
 		return SlateTexture.IsValid() ? FIntPoint(SlateTexture.Pin()->GetWidth(), SlateTexture.Pin()->GetHeight()) : FIntPoint();
 	}
 
-	virtual class FSlateShaderResource* GetViewportRenderTargetTexture() const OVERRIDE
+	virtual class FSlateShaderResource* GetViewportRenderTargetTexture() const override
 	{
 		return SlateTexture.IsValid() ? SlateTexture.Pin().Get() : NULL;
 	}
 
-	virtual bool RequiresVsync() const OVERRIDE
+	virtual bool RequiresVsync() const override
 	{
 		return false;
 	}
diff --git a/Engine/Source/Runtime/MovieSceneCore/Classes/MovieScenePropertyTrack.h b/Engine/Source/Runtime/MovieSceneCore/Classes/MovieScenePropertyTrack.h
--- a/Engine/Source/Runtime/MovieSceneCore/Classes/MovieScenePropertyTrack.h
+++ b/Engine/Source/Runtime/MovieSceneCore/Classes/MovieScenePropertyTrack.h
@@ -15,14 +15,14 @@ class MOVIESCENECORE_API UMovieScenePropertyTrack : public UMovieSceneTrack
 
 public:
 	/** UMovieSceneTrack interface */
-	virtual FName GetTrackName() const OVERRIDE { return PropertyName; }
-	virtual void RemoveAllAnimationData() OVERRIDE;
-	virtual bool HasSection( UMovieSceneSection* Section ) const OVERRIDE;
-	virtual void RemoveSection( UMovieSceneSection* Section ) OVERRIDE;
-	virtual bool IsEmpty() const OVERRIDE;
-	virtual TRange<float> GetSectionBoundaries() const OVERRIDE;
-	virtual TArray<UMovieSceneSection*> GetAllSections() const OVERRIDE;
-	virtual bool HasShowableData() const OVERRIDE {return bSectionsAreShowable;}
+	virtual FName GetTrackName() const override { return PropertyName; }
+	virtual void RemoveAllAnimationData() override;
+	virtual bool HasSection( UMovieSceneSection* Section ) const override;
+	virtual void RemoveSection( UMovieSceneSection* Section ) override;
+	virtual bool IsEmpty() const override;
+	virtual TRange<float> GetSectionBoundaries() const override;
+	virtual TArray<UMovieSceneSection*> GetAllSections() const override;
+	virtual bool HasShowableData() const override {return bSectionsAreShowable;}
 	
 	/**
 	 * Sets the property name for this animatable property
diff --git a/Engine/Source/Runtime/MovieSceneCore/Classes/RuntimeMovieScenePlayer.h b/Engine/Source/Runtime/MovieSceneCore/Classes/RuntimeMovieScenePlayer.h
--- a/Engine/Source/Runtime/MovieSceneCore/Classes/RuntimeMovieScenePlayer.h
+++ b/Engine/Source/Runtime/MovieSceneCore/Classes/RuntimeMovieScenePlayer.h
@@ -58,14 +58,14 @@ protected:
 	/** IMovieScenePlayer interface */
 	virtual void SpawnActorsForMovie( TSharedRef<FMovieSceneInstance> MovieSceneInstance );
 	virtual void DestroyActorsForMovie( TSharedRef<FMovieSceneInstance> MovieSceneInstance );
-	virtual void GetRuntimeObjects( TSharedRef<FMovieSceneInstance> MovieSceneInstance, const FGuid& ObjectHandle, TArray< UObject* >& OutObjects ) const OVERRIDE;
-	virtual void UpdateViewports(AActor* ActorToViewThrough) const OVERRIDE {}
-	virtual EMovieScenePlayerStatus::Type GetPlaybackStatus() const OVERRIDE;
-	virtual void AddMovieSceneInstance( class UMovieSceneSection& MovieSceneSection, TSharedRef<FMovieSceneInstance> InstanceToAdd ) OVERRIDE;
-	virtual void RemoveMovieSceneInstance( class UMovieSceneSection& MovieSceneSection, TSharedRef<FMovieSceneInstance> InstanceToRemove ) OVERRIDE;
-	virtual TSharedRef<FMovieSceneInstance> GetRootMovieSceneInstance() const OVERRIDE { return RootMovieSceneInstance.ToSharedRef(); }
+	virtual void GetRuntimeObjects( TSharedRef<FMovieSceneInstance> MovieSceneInstance, const FGuid& ObjectHandle, TArray< UObject* >& OutObjects ) const override;
+	virtual void UpdateViewports(AActor* ActorToViewThrough) const override {}
+	virtual EMovieScenePlayerStatus::Type GetPlaybackStatus() const override;
+	virtual void AddMovieSceneInstance( class UMovieSceneSection& MovieSceneSection, TSharedRef<FMovieSceneInstance> InstanceToAdd ) override;
+	virtual void RemoveMovieSceneInstance( class UMovieSceneSection& MovieSceneSection, TSharedRef<FMovieSceneInstance> InstanceToRemove ) override;
+	virtual TSharedRef<FMovieSceneInstance> GetRootMovieSceneInstance() const override { return RootMovieSceneInstance.ToSharedRef(); }
 	/** IRuntimeMovieScenePlayerInterface */
-	virtual void Tick( const float DeltaSeconds ) OVERRIDE;
+	virtual void Tick( const float DeltaSeconds ) override;
 
 
 private:
diff --git a/Engine/Source/Runtime/MovieSceneCoreTypes/Classes/MovieSceneAnimationSection.h b/Engine/Source/Runtime/MovieSceneCoreTypes/Classes/MovieSceneAnimationSection.h
--- a/Engine/Source/Runtime/MovieSceneCoreTypes/Classes/MovieSceneAnimationSection.h
+++ b/Engine/Source/Runtime/MovieSceneCoreTypes/Classes/MovieSceneAnimationSection.h
@@ -34,9 +34,9 @@ public:
 	float GetAnimationDilationFactor() const {return AnimationDilationFactor;}
 
 	/** MovieSceneSection interface */
-	virtual void MoveSection( float DeltaPosition ) OVERRIDE;
-	virtual void DilateSection( float DilationFactor, float Origin ) OVERRIDE;
-	virtual void GetSnapTimes(TArray<float>& OutSnapTimes, bool bGetSectionBorders) const OVERRIDE;
+	virtual void MoveSection( float DeltaPosition ) override;
+	virtual void DilateSection( float DilationFactor, float Origin ) override;
+	virtual void GetSnapTimes(TArray<float>& OutSnapTimes, bool bGetSectionBorders) const override;
 
 private:
 	/** The animation sequence this section has */
diff --git a/Engine/Source/Runtime/MovieSceneCoreTypes/Classes/MovieSceneAnimationTrack.h b/Engine/Source/Runtime/MovieSceneCoreTypes/Classes/MovieSceneAnimationTrack.h
--- a/Engine/Source/Runtime/MovieSceneCoreTypes/Classes/MovieSceneAnimationTrack.h
+++ b/Engine/Source/Runtime/MovieSceneCoreTypes/Classes/MovieSceneAnimationTrack.h
@@ -17,15 +17,15 @@ class UMovieSceneAnimationTrack : public UMovieSceneTrack
 
 public:
 	/** UMovieSceneTrack interface */
-	virtual FName GetTrackName() const OVERRIDE;
-	virtual TSharedPtr<IMovieSceneTrackInstance> CreateInstance() OVERRIDE;
-	virtual void RemoveAllAnimationData() OVERRIDE;
-	virtual bool HasSection( UMovieSceneSection* Section ) const OVERRIDE;
-	virtual void RemoveSection( UMovieSceneSection* Section ) OVERRIDE;
-	virtual bool IsEmpty() const OVERRIDE;
-	virtual TRange<float> GetSectionBoundaries() const OVERRIDE;
-	virtual bool SupportsMultipleRows() const OVERRIDE { return true; }
-	virtual TArray<UMovieSceneSection*> GetAllSections() const OVERRIDE;
+	virtual FName GetTrackName() const override;
+	virtual TSharedPtr<IMovieSceneTrackInstance> CreateInstance() override;
+	virtual void RemoveAllAnimationData() override;
+	virtual bool HasSection( UMovieSceneSection* Section ) const override;
+	virtual void RemoveSection( UMovieSceneSection* Section ) override;
+	virtual bool IsEmpty() const override;
+	virtual TRange<float> GetSectionBoundaries() const override;
+	virtual bool SupportsMultipleRows() const override { return true; }
+	virtual TArray<UMovieSceneSection*> GetAllSections() const override;
 
 	/** Adds a new animation to this track */
 	virtual void AddNewAnimation(float KeyTime, class UAnimSequence* AnimSequence);
diff --git a/Engine/Source/Runtime/MovieSceneCoreTypes/Classes/MovieSceneAudioSection.h b/Engine/Source/Runtime/MovieSceneCoreTypes/Classes/MovieSceneAudioSection.h
--- a/Engine/Source/Runtime/MovieSceneCoreTypes/Classes/MovieSceneAudioSection.h
+++ b/Engine/Source/Runtime/MovieSceneCoreTypes/Classes/MovieSceneAudioSection.h
@@ -53,9 +53,9 @@ public:
 	}
 
 	/** MovieSceneSection interface */
-	virtual void MoveSection( float DeltaPosition ) OVERRIDE;
-	virtual void DilateSection( float DilationFactor, float Origin ) OVERRIDE;
-	virtual void GetSnapTimes(TArray<float>& OutSnapTimes, bool bGetSectionBorders) const OVERRIDE;
+	virtual void MoveSection( float DeltaPosition ) override;
+	virtual void DilateSection( float DilationFactor, float Origin ) override;
+	virtual void GetSnapTimes(TArray<float>& OutSnapTimes, bool bGetSectionBorders) const override;
 
 private:
 	/** The sound cue or wave that this section plays */
diff --git a/Engine/Source/Runtime/MovieSceneCoreTypes/Classes/MovieSceneAudioTrack.h b/Engine/Source/Runtime/MovieSceneCoreTypes/Classes/MovieSceneAudioTrack.h
--- a/Engine/Source/Runtime/MovieSceneCoreTypes/Classes/MovieSceneAudioTrack.h
+++ b/Engine/Source/Runtime/MovieSceneCoreTypes/Classes/MovieSceneAudioTrack.h
@@ -22,15 +22,15 @@ class UMovieSceneAudioTrack : public UMovieSceneTrack
 
 public:
 	/** UMovieSceneTrack interface */
-	virtual FName GetTrackName() const OVERRIDE;
-	virtual TSharedPtr<IMovieSceneTrackInstance> CreateInstance() OVERRIDE;
-	virtual void RemoveAllAnimationData() OVERRIDE;
-	virtual bool HasSection( UMovieSceneSection* Section ) const OVERRIDE;
-	virtual void RemoveSection( UMovieSceneSection* Section ) OVERRIDE;
-	virtual bool IsEmpty() const OVERRIDE;
-	virtual TRange<float> GetSectionBoundaries() const OVERRIDE;
-	virtual bool SupportsMultipleRows() const OVERRIDE { return true; }
-	virtual TArray<UMovieSceneSection*> GetAllSections() const OVERRIDE;
+	virtual FName GetTrackName() const override;
+	virtual TSharedPtr<IMovieSceneTrackInstance> CreateInstance() override;
+	virtual void RemoveAllAnimationData() override;
+	virtual bool HasSection( UMovieSceneSection* Section ) const override;
+	virtual void RemoveSection( UMovieSceneSection* Section ) override;
+	virtual bool IsEmpty() const override;
+	virtual TRange<float> GetSectionBoundaries() const override;
+	virtual bool SupportsMultipleRows() const override { return true; }
+	virtual TArray<UMovieSceneSection*> GetAllSections() const override;
 
 	/** Adds a new sound cue to the audio */
 	virtual void AddNewSound(class USoundBase* Sound, float Time);
diff --git a/Engine/Source/Runtime/MovieSceneCoreTypes/Classes/MovieSceneBoolSection.h b/Engine/Source/Runtime/MovieSceneCoreTypes/Classes/MovieSceneBoolSection.h
--- a/Engine/Source/Runtime/MovieSceneCoreTypes/Classes/MovieSceneBoolSection.h
+++ b/Engine/Source/Runtime/MovieSceneCoreTypes/Classes/MovieSceneBoolSection.h
@@ -39,8 +39,8 @@ public:
 	/**
 	 * UMovieSceneSection interface 
 	 */
-	virtual void MoveSection( float DeltaPosition ) OVERRIDE;
-	virtual void DilateSection( float DilationFactor, float Origin ) OVERRIDE;
+	virtual void MoveSection( float DeltaPosition ) override;
+	virtual void DilateSection( float DilationFactor, float Origin ) override;
 
 	/** Gets all the keys of this boolean section */
 	FIntegralCurve& GetCurve() { return BoolCurve; }
diff --git a/Engine/Source/Runtime/MovieSceneCoreTypes/Classes/MovieSceneBoolTrack.h b/Engine/Source/Runtime/MovieSceneCoreTypes/Classes/MovieSceneBoolTrack.h
--- a/Engine/Source/Runtime/MovieSceneCoreTypes/Classes/MovieSceneBoolTrack.h
+++ b/Engine/Source/Runtime/MovieSceneCoreTypes/Classes/MovieSceneBoolTrack.h
@@ -17,8 +17,8 @@ class UMovieSceneBoolTrack : public UMovieScenePropertyTrack
 
 public:
 	/** UMovieSceneTrack interface */
-	virtual UMovieSceneSection* CreateNewSection() OVERRIDE;
-	virtual TSharedPtr<IMovieSceneTrackInstance> CreateInstance() OVERRIDE;
+	virtual UMovieSceneSection* CreateNewSection() override;
+	virtual TSharedPtr<IMovieSceneTrackInstance> CreateInstance() override;
 
 	/**
 	 * Adds a key to a section.  Will create the section if it doesn't exist
diff --git a/Engine/Source/Runtime/MovieSceneCoreTypes/Classes/MovieSceneColorSection.h b/Engine/Source/Runtime/MovieSceneCoreTypes/Classes/MovieSceneColorSection.h
--- a/Engine/Source/Runtime/MovieSceneCoreTypes/Classes/MovieSceneColorSection.h
+++ b/Engine/Source/Runtime/MovieSceneCoreTypes/Classes/MovieSceneColorSection.h
@@ -13,8 +13,8 @@ class UMovieSceneColorSection : public UMovieSceneSection
 	GENERATED_UCLASS_BODY()
 public:
 	/** MovieSceneSection interface */
-	virtual void MoveSection( float DeltaPosition ) OVERRIDE;
-	virtual void DilateSection( float DilationFactor, float Origin ) OVERRIDE;
+	virtual void MoveSection( float DeltaPosition ) override;
+	virtual void DilateSection( float DilationFactor, float Origin ) override;
 
 	/**
 	 * Updates this section
diff --git a/Engine/Source/Runtime/MovieSceneCoreTypes/Classes/MovieSceneColorTrack.h b/Engine/Source/Runtime/MovieSceneCoreTypes/Classes/MovieSceneColorTrack.h
--- a/Engine/Source/Runtime/MovieSceneCoreTypes/Classes/MovieSceneColorTrack.h
+++ b/Engine/Source/Runtime/MovieSceneCoreTypes/Classes/MovieSceneColorTrack.h
@@ -17,8 +17,8 @@ class UMovieSceneColorTrack : public UMovieScenePropertyTrack
 
 public:
 	/** UMovieSceneTrack interface */
-	virtual UMovieSceneSection* CreateNewSection() OVERRIDE;
-	virtual TSharedPtr<IMovieSceneTrackInstance> CreateInstance() OVERRIDE;
+	virtual UMovieSceneSection* CreateNewSection() override;
+	virtual TSharedPtr<IMovieSceneTrackInstance> CreateInstance() override;
 
 	/**
 	 * Adds a key to a section.  Will create the section if it doesn't exist
diff --git a/Engine/Source/Runtime/MovieSceneCoreTypes/Classes/MovieSceneParticleTrack.h b/Engine/Source/Runtime/MovieSceneCoreTypes/Classes/MovieSceneParticleTrack.h
--- a/Engine/Source/Runtime/MovieSceneCoreTypes/Classes/MovieSceneParticleTrack.h
+++ b/Engine/Source/Runtime/MovieSceneCoreTypes/Classes/MovieSceneParticleTrack.h
@@ -17,15 +17,15 @@ class UMovieSceneParticleTrack : public UMovieSceneTrack
 
 public:
 	/** UMovieSceneTrack interface */
-	virtual FName GetTrackName() const OVERRIDE;
-	virtual TSharedPtr<IMovieSceneTrackInstance> CreateInstance() OVERRIDE;
-	virtual void RemoveAllAnimationData() OVERRIDE;
-	virtual bool HasSection( UMovieSceneSection* Section ) const OVERRIDE;
-	virtual void RemoveSection( UMovieSceneSection* Section ) OVERRIDE;
-	virtual bool IsEmpty() const OVERRIDE;
-	virtual TRange<float> GetSectionBoundaries() const OVERRIDE;
-	virtual bool SupportsMultipleRows() const OVERRIDE { return false; }
-	virtual TArray<UMovieSceneSection*> GetAllSections() const OVERRIDE;
+	virtual FName GetTrackName() const override;
+	virtual TSharedPtr<IMovieSceneTrackInstance> CreateInstance() override;
+	virtual void RemoveAllAnimationData() override;
+	virtual bool HasSection( UMovieSceneSection* Section ) const override;
+	virtual void RemoveSection( UMovieSceneSection* Section ) override;
+	virtual bool IsEmpty() const override;
+	virtual TRange<float> GetSectionBoundaries() const override;
+	virtual bool SupportsMultipleRows() const override { return false; }
+	virtual TArray<UMovieSceneSection*> GetAllSections() const override;
 	
 	virtual void AddNewParticleSystem(float KeyTime, bool bTrigger);
 
diff --git a/Engine/Source/Runtime/MovieSceneCoreTypes/Classes/SubMovieSceneTrack.h b/Engine/Source/Runtime/MovieSceneCoreTypes/Classes/SubMovieSceneTrack.h
--- a/Engine/Source/Runtime/MovieSceneCoreTypes/Classes/SubMovieSceneTrack.h
+++ b/Engine/Source/Runtime/MovieSceneCoreTypes/Classes/SubMovieSceneTrack.h
@@ -14,14 +14,14 @@ class USubMovieSceneTrack : public UMovieSceneTrack
 	GENERATED_UCLASS_BODY()
 public:
 	/** UMovieSceneTrack interface */
-	virtual UMovieSceneSection* CreateNewSection() OVERRIDE;
-	virtual FName GetTrackName() const OVERRIDE;
-	virtual TSharedPtr<IMovieSceneTrackInstance> CreateInstance() OVERRIDE;
-	virtual TArray<UMovieSceneSection*> GetAllSections() const OVERRIDE;
-	virtual void RemoveSection( UMovieSceneSection* Section ) OVERRIDE;
-	virtual bool IsEmpty() const OVERRIDE;
-	virtual TRange<float> GetSectionBoundaries() const OVERRIDE;
-	virtual bool HasSection( UMovieSceneSection* Section ) const OVERRIDE;
+	virtual UMovieSceneSection* CreateNewSection() override;
+	virtual FName GetTrackName() const override;
+	virtual TSharedPtr<IMovieSceneTrackInstance> CreateInstance() override;
+	virtual TArray<UMovieSceneSection*> GetAllSections() const override;
+	virtual void RemoveSection( UMovieSceneSection* Section ) override;
+	virtual bool IsEmpty() const override;
+	virtual TRange<float> GetSectionBoundaries() const override;
+	virtual bool HasSection( UMovieSceneSection* Section ) const override;
 
 	/**
 	 * Adds a movie scene section
diff --git a/Engine/Source/Runtime/MovieSceneCoreTypes/Private/Instances/MovieSceneAnimationTrackInstance.h b/Engine/Source/Runtime/MovieSceneCoreTypes/Private/Instances/MovieSceneAnimationTrackInstance.h
--- a/Engine/Source/Runtime/MovieSceneCoreTypes/Private/Instances/MovieSceneAnimationTrackInstance.h
+++ b/Engine/Source/Runtime/MovieSceneCoreTypes/Private/Instances/MovieSceneAnimationTrackInstance.h
@@ -12,8 +12,8 @@ public:
 	virtual ~FMovieSceneAnimationTrackInstance();
 
 	/** IMovieSceneTrackInstance interface */
-	virtual void Update( float Position, float LastPosition, const TArray<UObject*>& RuntimeObjects, class IMovieScenePlayer& Player ) OVERRIDE;
-	virtual void RefreshInstance( IMovieScenePlayer& Player ) OVERRIDE {}
+	virtual void Update( float Position, float LastPosition, const TArray<UObject*>& RuntimeObjects, class IMovieScenePlayer& Player ) override;
+	virtual void RefreshInstance( IMovieScenePlayer& Player ) override {}
 
 private:
 	/** Track that is being instanced */
diff --git a/Engine/Source/Runtime/MovieSceneCoreTypes/Private/Instances/MovieSceneAudioTrackInstance.h b/Engine/Source/Runtime/MovieSceneCoreTypes/Private/Instances/MovieSceneAudioTrackInstance.h
--- a/Engine/Source/Runtime/MovieSceneCoreTypes/Private/Instances/MovieSceneAudioTrackInstance.h
+++ b/Engine/Source/Runtime/MovieSceneCoreTypes/Private/Instances/MovieSceneAudioTrackInstance.h
@@ -11,8 +11,8 @@ public:
 	FMovieSceneAudioTrackInstance( UMovieSceneAudioTrack& InAudioTrack );
 
 	/** IMovieSceneTrackInstance interface */
-	virtual void Update( float Position, float LastPosition, const TArray<UObject*>& RuntimeObjects, class IMovieScenePlayer& Player ) OVERRIDE;
-	virtual void RefreshInstance( class IMovieScenePlayer& Player ) OVERRIDE {}
+	virtual void Update( float Position, float LastPosition, const TArray<UObject*>& RuntimeObjects, class IMovieScenePlayer& Player ) override;
+	virtual void RefreshInstance( class IMovieScenePlayer& Player ) override {}
 private:
 	/** Plays the sound of the given section at the given time */
 	void PlaySound(class UMovieSceneAudioSection* AudioSection, class UAudioComponent* Component, float Time);
diff --git a/Engine/Source/Runtime/MovieSceneCoreTypes/Private/Instances/MovieSceneBoolTrackInstance.h b/Engine/Source/Runtime/MovieSceneCoreTypes/Private/Instances/MovieSceneBoolTrackInstance.h
--- a/Engine/Source/Runtime/MovieSceneCoreTypes/Private/Instances/MovieSceneBoolTrackInstance.h
+++ b/Engine/Source/Runtime/MovieSceneCoreTypes/Private/Instances/MovieSceneBoolTrackInstance.h
@@ -11,8 +11,8 @@ public:
 	FMovieSceneBoolTrackInstance( UMovieSceneBoolTrack& InBoolTrack );
 
 	/** IMovieSceneTrackInstance interface */
-	virtual void Update( float Position, float LastPosition, const TArray<UObject*>& RuntimeObjects, class IMovieScenePlayer& Player ) OVERRIDE;
-	virtual void RefreshInstance( class IMovieScenePlayer& Player ) OVERRIDE {}
+	virtual void Update( float Position, float LastPosition, const TArray<UObject*>& RuntimeObjects, class IMovieScenePlayer& Player ) override;
+	virtual void RefreshInstance( class IMovieScenePlayer& Player ) override {}
 private:
 	/** Bool track that is being instanced */
 	UMovieSceneBoolTrack* BoolTrack;
diff --git a/Engine/Source/Runtime/MovieSceneCoreTypes/Private/Instances/MovieSceneFloatTrackInstance.h b/Engine/Source/Runtime/MovieSceneCoreTypes/Private/Instances/MovieSceneFloatTrackInstance.h
--- a/Engine/Source/Runtime/MovieSceneCoreTypes/Private/Instances/MovieSceneFloatTrackInstance.h
+++ b/Engine/Source/Runtime/MovieSceneCoreTypes/Private/Instances/MovieSceneFloatTrackInstance.h
@@ -11,8 +11,8 @@ public:
 	FMovieSceneFloatTrackInstance( UMovieSceneFloatTrack& InFloatTrack );
 
 	/** IMovieSceneTrackInstance interface */
-	virtual void Update( float Position, float LastPosition, const TArray<UObject*>& RuntimeObjects, class IMovieScenePlayer& Player ) OVERRIDE;
-	virtual void RefreshInstance( IMovieScenePlayer& Player ) OVERRIDE {}
+	virtual void Update( float Position, float LastPosition, const TArray<UObject*>& RuntimeObjects, class IMovieScenePlayer& Player ) override;
+	virtual void RefreshInstance( IMovieScenePlayer& Player ) override {}
 private:
 	/** Track that is being instanced */
 	UMovieSceneFloatTrack* FloatTrack;
diff --git a/Engine/Source/Runtime/MovieSceneCoreTypes/Private/Instances/SubMovieSceneTrackInstance.h b/Engine/Source/Runtime/MovieSceneCoreTypes/Private/Instances/SubMovieSceneTrackInstance.h
--- a/Engine/Source/Runtime/MovieSceneCoreTypes/Private/Instances/SubMovieSceneTrackInstance.h
+++ b/Engine/Source/Runtime/MovieSceneCoreTypes/Private/Instances/SubMovieSceneTrackInstance.h
@@ -11,8 +11,8 @@ public:
 	FSubMovieSceneTrackInstance( USubMovieSceneTrack& InTrack );
 
 	/** IMovieSceneTrackInstance interface */
-	virtual void Update( float Position, float LastPosition, const TArray<UObject*>& RuntimeObjects, class IMovieScenePlayer& Player ) OVERRIDE;
-	virtual void RefreshInstance( class IMovieScenePlayer& Player ) OVERRIDE;
+	virtual void Update( float Position, float LastPosition, const TArray<UObject*>& RuntimeObjects, class IMovieScenePlayer& Player ) override;
+	virtual void RefreshInstance( class IMovieScenePlayer& Player ) override;
 private:
 	/** Track that is being instanced */
 	TWeakObjectPtr<class USubMovieSceneTrack> SubMovieSceneTrack;
diff --git a/Engine/Source/Runtime/NetworkFile/Private/NetworkPlatformFile.cpp b/Engine/Source/Runtime/NetworkFile/Private/NetworkPlatformFile.cpp
--- a/Engine/Source/Runtime/NetworkFile/Private/NetworkPlatformFile.cpp
+++ b/Engine/Source/Runtime/NetworkFile/Private/NetworkPlatformFile.cpp
@@ -1084,7 +1084,7 @@ void FNetworkPlatformFile::ConvertServerFilenameToClientFilename(FString& Filena
 class FNetworkFileModule : public IPlatformFileModule
 {
 public:
-	virtual IPlatformFile* GetPlatformFile() OVERRIDE
+	virtual IPlatformFile* GetPlatformFile() override
 	{
 		static TScopedPointer<IPlatformFile> AutoDestroySingleton(new FNetworkPlatformFile());
 		return AutoDestroySingleton.GetOwnedPointer();
diff --git a/Engine/Source/Runtime/NetworkFile/Public/NetworkPlatformFile.h b/Engine/Source/Runtime/NetworkFile/Public/NetworkPlatformFile.h
--- a/Engine/Source/Runtime/NetworkFile/Public/NetworkPlatformFile.h
+++ b/Engine/Source/Runtime/NetworkFile/Public/NetworkPlatformFile.h
@@ -38,16 +38,16 @@ public:
 	/** Destructor */
 	virtual ~FNetworkPlatformFile();
 
-	virtual bool ShouldBeUsed(IPlatformFile* Inner, const TCHAR* CmdLine) const OVERRIDE;
-	virtual bool Initialize(IPlatformFile* Inner, const TCHAR* CmdLine) OVERRIDE;
-	virtual void InitializeAfterSetActive() OVERRIDE;
+	virtual bool ShouldBeUsed(IPlatformFile* Inner, const TCHAR* CmdLine) const override;
+	virtual bool Initialize(IPlatformFile* Inner, const TCHAR* CmdLine) override;
+	virtual void InitializeAfterSetActive() override;
 
-	virtual IPlatformFile* GetLowerLevel() OVERRIDE
+	virtual IPlatformFile* GetLowerLevel() override
 	{
 		return InnerPlatformFile;
 	}
 
-	virtual const TCHAR* GetName() const OVERRIDE
+	virtual const TCHAR* GetName() const override
 	{
 		return FNetworkPlatformFile::GetTypeName();
 	}
@@ -57,56 +57,56 @@ public:
 		return bIsUsable;
 	}
 
-	virtual bool		FileExists(const TCHAR* Filename) OVERRIDE
+	virtual bool		FileExists(const TCHAR* Filename) override
 	{
 		FFileInfo Info;
 		GetFileInfo(Filename, Info);
 		return Info.FileExists;
 	}
-	virtual int64		FileSize(const TCHAR* Filename) OVERRIDE
+	virtual int64		FileSize(const TCHAR* Filename) override
 	{
 		FFileInfo Info;
 		GetFileInfo(Filename, Info);
 		return Info.Size;
 	}
-	virtual bool		DeleteFile(const TCHAR* Filename) OVERRIDE;
-	virtual bool		IsReadOnly(const TCHAR* Filename) OVERRIDE
+	virtual bool		DeleteFile(const TCHAR* Filename) override;
+	virtual bool		IsReadOnly(const TCHAR* Filename) override
 	{
 		FFileInfo Info;
 		GetFileInfo(Filename, Info);
 		return Info.ReadOnly;
 	}
-	virtual bool		MoveFile(const TCHAR* To, const TCHAR* From) OVERRIDE;
-	virtual bool		SetReadOnly(const TCHAR* Filename, bool bNewReadOnlyValue) OVERRIDE;
-	virtual FDateTime	GetTimeStamp(const TCHAR* Filename) OVERRIDE
+	virtual bool		MoveFile(const TCHAR* To, const TCHAR* From) override;
+	virtual bool		SetReadOnly(const TCHAR* Filename, bool bNewReadOnlyValue) override;
+	virtual FDateTime	GetTimeStamp(const TCHAR* Filename) override
 	{
 		FFileInfo Info;
 		GetFileInfo(Filename, Info);
 		return Info.TimeStamp;
 	}
-	virtual void		SetTimeStamp(const TCHAR* Filename, FDateTime DateTime) OVERRIDE;
-	virtual FDateTime	GetAccessTimeStamp(const TCHAR* Filename) OVERRIDE
+	virtual void		SetTimeStamp(const TCHAR* Filename, FDateTime DateTime) override;
+	virtual FDateTime	GetAccessTimeStamp(const TCHAR* Filename) override
 	{
 		FFileInfo Info;
 		GetFileInfo(Filename, Info);
 		return Info.AccessTimeStamp;
 	}
-	virtual IFileHandle*	OpenRead(const TCHAR* Filename) OVERRIDE;
-	virtual IFileHandle*	OpenWrite(const TCHAR* Filename, bool bAppend = false, bool bAllowRead = false) OVERRIDE;
-	virtual bool		DirectoryExists(const TCHAR* Directory) OVERRIDE;
-	virtual bool		CreateDirectoryTree(const TCHAR* Directory) OVERRIDE;
-	virtual bool		CreateDirectory(const TCHAR* Directory) OVERRIDE;
-	virtual bool		DeleteDirectory(const TCHAR* Directory) OVERRIDE;
+	virtual IFileHandle*	OpenRead(const TCHAR* Filename) override;
+	virtual IFileHandle*	OpenWrite(const TCHAR* Filename, bool bAppend = false, bool bAllowRead = false) override;
+	virtual bool		DirectoryExists(const TCHAR* Directory) override;
+	virtual bool		CreateDirectoryTree(const TCHAR* Directory) override;
+	virtual bool		CreateDirectory(const TCHAR* Directory) override;
+	virtual bool		DeleteDirectory(const TCHAR* Directory) override;
 
-	virtual bool		IterateDirectory(const TCHAR* Directory, IPlatformFile::FDirectoryVisitor& Visitor) OVERRIDE;
-	virtual bool		IterateDirectoryRecursively(const TCHAR* Directory, IPlatformFile::FDirectoryVisitor& Visitor) OVERRIDE;
-	virtual bool		DeleteDirectoryRecursively(const TCHAR* Directory) OVERRIDE;
-	virtual bool		CopyFile(const TCHAR* To, const TCHAR* From) OVERRIDE;
+	virtual bool		IterateDirectory(const TCHAR* Directory, IPlatformFile::FDirectoryVisitor& Visitor) override;
+	virtual bool		IterateDirectoryRecursively(const TCHAR* Directory, IPlatformFile::FDirectoryVisitor& Visitor) override;
+	virtual bool		DeleteDirectoryRecursively(const TCHAR* Directory) override;
+	virtual bool		CopyFile(const TCHAR* To, const TCHAR* From) override;
 
-	virtual FString ConvertToAbsolutePathForExternalAppForRead( const TCHAR* Filename ) OVERRIDE;
-	virtual FString ConvertToAbsolutePathForExternalAppForWrite( const TCHAR* Filename ) OVERRIDE;
+	virtual FString ConvertToAbsolutePathForExternalAppForRead( const TCHAR* Filename ) override;
+	virtual FString ConvertToAbsolutePathForExternalAppForWrite( const TCHAR* Filename ) override;
 
-	virtual bool SendMessageToServer(const TCHAR* Message, IPlatformFile::IFileServerMessageHandler* Handler) OVERRIDE;
+	virtual bool SendMessageToServer(const TCHAR* Message, IPlatformFile::IFileServerMessageHandler* Handler) override;
 
 
 
@@ -227,11 +227,11 @@ public:
 	{
 	}
 
-	virtual int64		Tell() OVERRIDE
+	virtual int64		Tell() override
 	{
 		return FilePos;
 	}
-	virtual bool		Seek(int64 NewPosition) OVERRIDE
+	virtual bool		Seek(int64 NewPosition) override
 	{
 		if (NewPosition >= 0 && NewPosition <= Size)
 		{
@@ -240,11 +240,11 @@ public:
 		}
 		return false;
 	}
-	virtual bool		SeekFromEnd(int64 NewPositionRelativeToEnd = 0) OVERRIDE
+	virtual bool		SeekFromEnd(int64 NewPositionRelativeToEnd = 0) override
 	{
 		return Seek(Size + NewPositionRelativeToEnd);
 	}
-	virtual bool		Read(uint8* Destination, int64 BytesToRead) OVERRIDE
+	virtual bool		Read(uint8* Destination, int64 BytesToRead) override
 	{
 		bool Result = false;
 		if (bReadable && BytesToRead >= 0 && BytesToRead + FilePos <= Size)
@@ -264,7 +264,7 @@ public:
 		}
 		return Result;
 	}
-	virtual bool		Write(const uint8* Source, int64 BytesToWrite) OVERRIDE
+	virtual bool		Write(const uint8* Source, int64 BytesToWrite) override
 	{
 		bool Result = false;
 		if (bWritable && BytesToWrite >= 0)
diff --git a/Engine/Source/Runtime/NetworkFileSystem/Private/NetworkFileServer.h b/Engine/Source/Runtime/NetworkFileSystem/Private/NetworkFileServer.h
--- a/Engine/Source/Runtime/NetworkFileSystem/Private/NetworkFileServer.h
+++ b/Engine/Source/Runtime/NetworkFileSystem/Private/NetworkFileServer.h
@@ -29,19 +29,19 @@ public:
 
 	// FRunnable Interface
 
-	virtual bool Init( ) OVERRIDE
+	virtual bool Init( ) override
 	{
 		return true;
 	}
 
-	virtual uint32 Run( ) OVERRIDE;
+	virtual uint32 Run( ) override;
 
-	virtual void Stop( ) OVERRIDE
+	virtual void Stop( ) override
 	{
 		StopRequested.Set(true);
 	}
 
-	virtual void Exit( ) OVERRIDE;
+	virtual void Exit( ) override;
 
 public:
 
diff --git a/Engine/Source/Runtime/NetworkFileSystem/Private/NetworkFileSystemModule.cpp b/Engine/Source/Runtime/NetworkFileSystem/Private/NetworkFileSystemModule.cpp
--- a/Engine/Source/Runtime/NetworkFileSystem/Private/NetworkFileSystemModule.cpp
+++ b/Engine/Source/Runtime/NetworkFileSystem/Private/NetworkFileSystemModule.cpp
@@ -20,7 +20,7 @@ public:
 
 	// INetworkFileSystemModule interface
 
-	virtual INetworkFileServer* CreateNetworkFileServer( int32 Port, const FFileRequestDelegate* InFileRequestDelegate, const FRecompileShadersDelegate* InRecompileShadersDelegate ) const OVERRIDE
+	virtual INetworkFileServer* CreateNetworkFileServer( int32 Port, const FFileRequestDelegate* InFileRequestDelegate, const FRecompileShadersDelegate* InRecompileShadersDelegate ) const override
 	{
 		if (Port < 0)
 		{
diff --git a/Engine/Source/Runtime/NullDrv/Public/NullDrv.h b/Engine/Source/Runtime/NullDrv/Public/NullDrv.h
--- a/Engine/Source/Runtime/NullDrv/Public/NullDrv.h
+++ b/Engine/Source/Runtime/NullDrv/Public/NullDrv.h
@@ -17,12 +17,12 @@ class FNullDynamicRHIModule : public IDynamicRHIModule
 public:
 
 	// IModuleInterface
-	virtual bool SupportsDynamicReloading() OVERRIDE { return false; }
+	virtual bool SupportsDynamicReloading() override { return false; }
 
 	// IDynamicRHIModule
-	virtual bool IsSupported() OVERRIDE;
+	virtual bool IsSupported() override;
 
-	virtual FDynamicRHI* CreateRHI() OVERRIDE
+	virtual FDynamicRHI* CreateRHI() override
 	{
 		return new FNullDynamicRHI();
 	}
diff --git a/Engine/Source/Runtime/Online/Android/OnlineSubsystemGooglePlay/Private/OnlineAchievementsInterfaceGooglePlay.h b/Engine/Source/Runtime/Online/Android/OnlineSubsystemGooglePlay/Private/OnlineAchievementsInterfaceGooglePlay.h
--- a/Engine/Source/Runtime/Online/Android/OnlineSubsystemGooglePlay/Private/OnlineAchievementsInterfaceGooglePlay.h
+++ b/Engine/Source/Runtime/Online/Android/OnlineSubsystemGooglePlay/Private/OnlineAchievementsInterfaceGooglePlay.h
@@ -56,14 +56,14 @@ private:
 public:
 
 	// Begin IOnlineAchievements interface
-	virtual void WriteAchievements(const FUniqueNetId& PlayerId, FOnlineAchievementsWriteRef& WriteObject, const FOnAchievementsWrittenDelegate& Delegate = FOnAchievementsWrittenDelegate()) OVERRIDE;
-	virtual void QueryAchievements(const FUniqueNetId& PlayerId, const FOnQueryAchievementsCompleteDelegate& Delegate = FOnQueryAchievementsCompleteDelegate()) OVERRIDE;
-	virtual void QueryAchievementDescriptions(const FUniqueNetId& PlayerId, const FOnQueryAchievementsCompleteDelegate& Delegate = FOnQueryAchievementsCompleteDelegate()) OVERRIDE ;
-	virtual EOnlineCachedResult::Type GetCachedAchievement(const FUniqueNetId& PlayerId, const FString& AchievementId, FOnlineAchievement& OutAchievement) OVERRIDE;
-	virtual EOnlineCachedResult::Type GetCachedAchievements(const FUniqueNetId& PlayerId, TArray<FOnlineAchievement>& OutAchievements) OVERRIDE;
-	virtual EOnlineCachedResult::Type GetCachedAchievementDescription(const FString& AchievementId, FOnlineAchievementDesc& OutAchievementDesc) OVERRIDE;
+	virtual void WriteAchievements(const FUniqueNetId& PlayerId, FOnlineAchievementsWriteRef& WriteObject, const FOnAchievementsWrittenDelegate& Delegate = FOnAchievementsWrittenDelegate()) override;
+	virtual void QueryAchievements(const FUniqueNetId& PlayerId, const FOnQueryAchievementsCompleteDelegate& Delegate = FOnQueryAchievementsCompleteDelegate()) override;
+	virtual void QueryAchievementDescriptions(const FUniqueNetId& PlayerId, const FOnQueryAchievementsCompleteDelegate& Delegate = FOnQueryAchievementsCompleteDelegate()) override ;
+	virtual EOnlineCachedResult::Type GetCachedAchievement(const FUniqueNetId& PlayerId, const FString& AchievementId, FOnlineAchievement& OutAchievement) override;
+	virtual EOnlineCachedResult::Type GetCachedAchievements(const FUniqueNetId& PlayerId, TArray<FOnlineAchievement>& OutAchievements) override;
+	virtual EOnlineCachedResult::Type GetCachedAchievementDescription(const FString& AchievementId, FOnlineAchievementDesc& OutAchievementDesc) override;
 #if !UE_BUILD_SHIPPING
-	virtual bool ResetAchievements( const FUniqueNetId& PlayerId ) OVERRIDE;
+	virtual bool ResetAchievements( const FUniqueNetId& PlayerId ) override;
 #endif // !UE_BUILD_SHIPPING
 	// End IOnlineAchievements interface
 
diff --git a/Engine/Source/Runtime/Online/Android/OnlineSubsystemGooglePlay/Private/OnlineAsyncTaskManagerGooglePlay.h b/Engine/Source/Runtime/Online/Android/OnlineSubsystemGooglePlay/Private/OnlineAsyncTaskManagerGooglePlay.h
--- a/Engine/Source/Runtime/Online/Android/OnlineSubsystemGooglePlay/Private/OnlineAsyncTaskManagerGooglePlay.h
+++ b/Engine/Source/Runtime/Online/Android/OnlineSubsystemGooglePlay/Private/OnlineAsyncTaskManagerGooglePlay.h
@@ -15,7 +15,7 @@ public:
 	FOnlineAsyncTaskManagerGooglePlay() {}
 
 	// FOnlineAsyncTaskManager
-	virtual void OnlineTick() OVERRIDE {}
+	virtual void OnlineTick() override {}
 };
 
 
diff --git a/Engine/Source/Runtime/Online/Android/OnlineSubsystemGooglePlay/Private/OnlineExternalUIInterfaceGooglePlay.h b/Engine/Source/Runtime/Online/Android/OnlineSubsystemGooglePlay/Private/OnlineExternalUIInterfaceGooglePlay.h
--- a/Engine/Source/Runtime/Online/Android/OnlineSubsystemGooglePlay/Private/OnlineExternalUIInterfaceGooglePlay.h
+++ b/Engine/Source/Runtime/Online/Android/OnlineSubsystemGooglePlay/Private/OnlineExternalUIInterfaceGooglePlay.h
@@ -17,13 +17,13 @@ class FOnlineExternalUIGooglePlay :
 public:
 	FOnlineExternalUIGooglePlay();
 
-	virtual bool ShowLoginUI(const int ControllerIndex, bool bShowOnlineOnly, const FOnLoginUIClosedDelegate& Delegate) OVERRIDE;
-	virtual bool ShowFriendsUI(int32 LocalUserNum) OVERRIDE;
-	virtual bool ShowInviteUI(int32 LocalUserNum) OVERRIDE;
-	virtual bool ShowAchievementsUI(int32 LocalUserNum) OVERRIDE;
-	virtual bool ShowLeaderboardUI(const FString& LeaderboardName) OVERRIDE;
-	virtual bool ShowWebURL(const FString& WebURL) OVERRIDE;
-	virtual bool ShowProfileUI(const FUniqueNetId& Requestor, const FUniqueNetId& Requestee, const FOnProfileUIClosedDelegate& Delegate) OVERRIDE;
+	virtual bool ShowLoginUI(const int ControllerIndex, bool bShowOnlineOnly, const FOnLoginUIClosedDelegate& Delegate) override;
+	virtual bool ShowFriendsUI(int32 LocalUserNum) override;
+	virtual bool ShowInviteUI(int32 LocalUserNum) override;
+	virtual bool ShowAchievementsUI(int32 LocalUserNum) override;
+	virtual bool ShowLeaderboardUI(const FString& LeaderboardName) override;
+	virtual bool ShowWebURL(const FString& WebURL) override;
+	virtual bool ShowProfileUI(const FUniqueNetId& Requestor, const FUniqueNetId& Requestee, const FOnProfileUIClosedDelegate& Delegate) override;
 };
 
 typedef TSharedPtr<FOnlineExternalUIGooglePlay, ESPMode::ThreadSafe> FOnlineExternalUIGooglePlayPtr;
diff --git a/Engine/Source/Runtime/Online/Android/OnlineSubsystemGooglePlay/Private/OnlineIdentityInterfaceGooglePlay.h b/Engine/Source/Runtime/Online/Android/OnlineSubsystemGooglePlay/Private/OnlineIdentityInterfaceGooglePlay.h
--- a/Engine/Source/Runtime/Online/Android/OnlineSubsystemGooglePlay/Private/OnlineIdentityInterfaceGooglePlay.h
+++ b/Engine/Source/Runtime/Online/Android/OnlineSubsystemGooglePlay/Private/OnlineIdentityInterfaceGooglePlay.h
@@ -29,17 +29,17 @@ PACKAGE_SCOPE:
 public:
 
 	// Begin IOnlineIdentity interface
-	virtual TSharedPtr<FUserOnlineAccount> GetUserAccount(const FUniqueNetId& UserId) const OVERRIDE;
-	virtual TArray<TSharedPtr<FUserOnlineAccount> > GetAllUserAccounts() const OVERRIDE;
-	virtual bool Login(int32 LocalUserNum, const FOnlineAccountCredentials& AccountCredentials) OVERRIDE;
-	virtual bool Logout(int32 LocalUserNum) OVERRIDE;
-	virtual bool AutoLogin(int32 LocalUserNum) OVERRIDE;
-	virtual TSharedPtr<FUniqueNetId> GetUniquePlayerId(int32 LocalUserNum) const OVERRIDE;
-	virtual TSharedPtr<FUniqueNetId> CreateUniquePlayerId(uint8* Bytes, int32 Size) OVERRIDE;
-	virtual TSharedPtr<FUniqueNetId> CreateUniquePlayerId(const FString& Str) OVERRIDE;
-	virtual ELoginStatus::Type GetLoginStatus(int32 LocalUserNum) const OVERRIDE;
-	virtual FString GetPlayerNickname(int32 LocalUserNum) const OVERRIDE;
-	virtual FString GetAuthToken(int32 LocalUserNum) const OVERRIDE;
+	virtual TSharedPtr<FUserOnlineAccount> GetUserAccount(const FUniqueNetId& UserId) const override;
+	virtual TArray<TSharedPtr<FUserOnlineAccount> > GetAllUserAccounts() const override;
+	virtual bool Login(int32 LocalUserNum, const FOnlineAccountCredentials& AccountCredentials) override;
+	virtual bool Logout(int32 LocalUserNum) override;
+	virtual bool AutoLogin(int32 LocalUserNum) override;
+	virtual TSharedPtr<FUniqueNetId> GetUniquePlayerId(int32 LocalUserNum) const override;
+	virtual TSharedPtr<FUniqueNetId> CreateUniquePlayerId(uint8* Bytes, int32 Size) override;
+	virtual TSharedPtr<FUniqueNetId> CreateUniquePlayerId(const FString& Str) override;
+	virtual ELoginStatus::Type GetLoginStatus(int32 LocalUserNum) const override;
+	virtual FString GetPlayerNickname(int32 LocalUserNum) const override;
+	virtual FString GetAuthToken(int32 LocalUserNum) const override;
 	// End IOnlineIdentity interface
 
 public:
diff --git a/Engine/Source/Runtime/Online/Android/OnlineSubsystemGooglePlay/Private/OnlineLeaderboardInterfaceGooglePlay.h b/Engine/Source/Runtime/Online/Android/OnlineSubsystemGooglePlay/Private/OnlineLeaderboardInterfaceGooglePlay.h
--- a/Engine/Source/Runtime/Online/Android/OnlineSubsystemGooglePlay/Private/OnlineLeaderboardInterfaceGooglePlay.h
+++ b/Engine/Source/Runtime/Online/Android/OnlineSubsystemGooglePlay/Private/OnlineLeaderboardInterfaceGooglePlay.h
@@ -19,12 +19,12 @@ class FOnlineLeaderboardsGooglePlay : public IOnlineLeaderboards
 public:
 	FOnlineLeaderboardsGooglePlay(class FOnlineSubsystemGooglePlay* InSubsystem);
 
-	virtual bool ReadLeaderboards(const TArray< TSharedRef<FUniqueNetId> >& Players, FOnlineLeaderboardReadRef& ReadObject) OVERRIDE;
-	virtual bool ReadLeaderboardsForFriends(int32 LocalUserNum, FOnlineLeaderboardReadRef& ReadObject) OVERRIDE;
-	virtual void FreeStats(FOnlineLeaderboardRead& ReadObject) OVERRIDE;
-	virtual bool WriteLeaderboards(const FName& SessionName, const FUniqueNetId& Player, FOnlineLeaderboardWrite& WriteObject) OVERRIDE;
-	virtual bool FlushLeaderboards(const FName& SessionName) OVERRIDE;
-	virtual bool WriteOnlinePlayerRatings(const FName& SessionName, int32 LeaderboardId, const TArray<FOnlinePlayerScore>& PlayerScores) OVERRIDE;
+	virtual bool ReadLeaderboards(const TArray< TSharedRef<FUniqueNetId> >& Players, FOnlineLeaderboardReadRef& ReadObject) override;
+	virtual bool ReadLeaderboardsForFriends(int32 LocalUserNum, FOnlineLeaderboardReadRef& ReadObject) override;
+	virtual void FreeStats(FOnlineLeaderboardRead& ReadObject) override;
+	virtual bool WriteLeaderboards(const FName& SessionName, const FUniqueNetId& Player, FOnlineLeaderboardWrite& WriteObject) override;
+	virtual bool FlushLeaderboards(const FName& SessionName) override;
+	virtual bool WriteOnlinePlayerRatings(const FName& SessionName, int32 LeaderboardId, const TArray<FOnlinePlayerScore>& PlayerScores) override;
 
 private:
 	/**
diff --git a/Engine/Source/Runtime/Online/Android/OnlineSubsystemGooglePlay/Public/OnlineSubsystemGooglePlay.h b/Engine/Source/Runtime/Online/Android/OnlineSubsystemGooglePlay/Public/OnlineSubsystemGooglePlay.h
--- a/Engine/Source/Runtime/Online/Android/OnlineSubsystemGooglePlay/Public/OnlineSubsystemGooglePlay.h
+++ b/Engine/Source/Runtime/Online/Android/OnlineSubsystemGooglePlay/Public/OnlineSubsystemGooglePlay.h
@@ -21,39 +21,39 @@ public:
 	virtual ~FOnlineSubsystemGooglePlay() {}
 
 	// Begin IOnlineSubsystem Interface
-	virtual IOnlineSessionPtr GetSessionInterface() const OVERRIDE;
-	virtual IOnlineFriendsPtr GetFriendsInterface() const OVERRIDE;
-	virtual IOnlineSharedCloudPtr GetSharedCloudInterface() const OVERRIDE;
-	virtual IOnlineUserCloudPtr GetUserCloudInterface() const OVERRIDE;
-	virtual IOnlineLeaderboardsPtr GetLeaderboardsInterface() const OVERRIDE;
-	virtual IOnlineVoicePtr GetVoiceInterface() const  OVERRIDE;
-	virtual IOnlineExternalUIPtr GetExternalUIInterface() const OVERRIDE;
-	virtual IOnlineTimePtr GetTimeInterface() const OVERRIDE;
-	virtual IOnlineIdentityPtr GetIdentityInterface() const OVERRIDE;
-	virtual IOnlinePartyPtr GetPartyInterface() const OVERRIDE;
-	virtual IOnlineTitleFilePtr GetTitleFileInterface() const OVERRIDE;
-	virtual IOnlineEntitlementsPtr GetEntitlementsInterface() const OVERRIDE;
-	virtual IOnlineStorePtr GetStoreInterface() const OVERRIDE;
-	virtual IOnlineEventsPtr GetEventsInterface() const OVERRIDE { return NULL; }
-	virtual IOnlineMessagePtr GetMessageInterface() const OVERRIDE { return NULL; }
-	virtual IOnlineSharingPtr GetSharingInterface() const OVERRIDE { return NULL; }
-	virtual IOnlineUserPtr GetUserInterface() const OVERRIDE { return NULL; }
-	virtual IOnlineAchievementsPtr GetAchievementsInterface() const OVERRIDE;
-	virtual IOnlinePresencePtr GetPresenceInterface() const OVERRIDE { return NULL; }
-	virtual class UObject* GetNamedInterface(FName InterfaceName) OVERRIDE { return NULL; }
-	virtual void SetNamedInterface(FName InterfaceName, class UObject* NewInterface) OVERRIDE {}
-	virtual bool IsDedicated() const OVERRIDE { return false; }
-	virtual bool IsServer() const OVERRIDE { return true; }
-	virtual void SetForceDedicated(bool bForce) OVERRIDE {}
-	virtual bool IsLocalPlayer(const FUniqueNetId& UniqueId) const OVERRIDE { return true; }
+	virtual IOnlineSessionPtr GetSessionInterface() const override;
+	virtual IOnlineFriendsPtr GetFriendsInterface() const override;
+	virtual IOnlineSharedCloudPtr GetSharedCloudInterface() const override;
+	virtual IOnlineUserCloudPtr GetUserCloudInterface() const override;
+	virtual IOnlineLeaderboardsPtr GetLeaderboardsInterface() const override;
+	virtual IOnlineVoicePtr GetVoiceInterface() const  override;
+	virtual IOnlineExternalUIPtr GetExternalUIInterface() const override;
+	virtual IOnlineTimePtr GetTimeInterface() const override;
+	virtual IOnlineIdentityPtr GetIdentityInterface() const override;
+	virtual IOnlinePartyPtr GetPartyInterface() const override;
+	virtual IOnlineTitleFilePtr GetTitleFileInterface() const override;
+	virtual IOnlineEntitlementsPtr GetEntitlementsInterface() const override;
+	virtual IOnlineStorePtr GetStoreInterface() const override;
+	virtual IOnlineEventsPtr GetEventsInterface() const override { return NULL; }
+	virtual IOnlineMessagePtr GetMessageInterface() const override { return NULL; }
+	virtual IOnlineSharingPtr GetSharingInterface() const override { return NULL; }
+	virtual IOnlineUserPtr GetUserInterface() const override { return NULL; }
+	virtual IOnlineAchievementsPtr GetAchievementsInterface() const override;
+	virtual IOnlinePresencePtr GetPresenceInterface() const override { return NULL; }
+	virtual class UObject* GetNamedInterface(FName InterfaceName) override { return NULL; }
+	virtual void SetNamedInterface(FName InterfaceName, class UObject* NewInterface) override {}
+	virtual bool IsDedicated() const override { return false; }
+	virtual bool IsServer() const override { return true; }
+	virtual void SetForceDedicated(bool bForce) override {}
+	virtual bool IsLocalPlayer(const FUniqueNetId& UniqueId) const override { return true; }
 
-	virtual bool Init() OVERRIDE;
-	virtual bool Shutdown() OVERRIDE;
-	virtual FString GetAppId() const OVERRIDE;
-	virtual bool Exec(class UWorld* InWorld, const TCHAR* Cmd, FOutputDevice& Ar) OVERRIDE;
+	virtual bool Init() override;
+	virtual bool Shutdown() override;
+	virtual FString GetAppId() const override;
+	virtual bool Exec(class UWorld* InWorld, const TCHAR* Cmd, FOutputDevice& Ar) override;
 	// End IOnlineSubsystem Interface
 
-	virtual bool Tick(float DeltaTime) OVERRIDE;
+	virtual bool Tick(float DeltaTime) override;
 	
 PACKAGE_SCOPE:
 
diff --git a/Engine/Source/Runtime/Online/Android/OnlineSubsystemGooglePlay/Public/OnlineSubsystemGooglePlayModule.h b/Engine/Source/Runtime/Online/Android/OnlineSubsystemGooglePlay/Public/OnlineSubsystemGooglePlayModule.h
--- a/Engine/Source/Runtime/Online/Android/OnlineSubsystemGooglePlay/Public/OnlineSubsystemGooglePlayModule.h
+++ b/Engine/Source/Runtime/Online/Android/OnlineSubsystemGooglePlay/Public/OnlineSubsystemGooglePlayModule.h
@@ -19,15 +19,15 @@ private:
 	virtual ~FOnlineSubsystemGooglePlayModule(){}
 
 	// Begin IModuleInterface interface
-	virtual void StartupModule() OVERRIDE;
-	virtual void ShutdownModule() OVERRIDE;
+	virtual void StartupModule() override;
+	virtual void ShutdownModule() override;
 
-	virtual bool SupportsDynamicReloading() OVERRIDE
+	virtual bool SupportsDynamicReloading() override
 	{
 		return false;
 	}
 
-	virtual bool SupportsAutomaticShutdown() OVERRIDE
+	virtual bool SupportsAutomaticShutdown() override
 	{
 		return false;
 	}
diff --git a/Engine/Source/Runtime/Online/HTTP/Private/Curl/CurlHttp.h b/Engine/Source/Runtime/Online/HTTP/Private/Curl/CurlHttp.h
--- a/Engine/Source/Runtime/Online/HTTP/Private/Curl/CurlHttp.h
+++ b/Engine/Source/Runtime/Online/HTTP/Private/Curl/CurlHttp.h
@@ -107,28 +107,28 @@ public:
 	friend class FCurlHttpResponse;
 
 	// Begin IHttpBase interface
-	virtual FString GetURL() OVERRIDE;
-	virtual FString GetURLParameter(const FString& ParameterName) OVERRIDE;
-	virtual FString GetHeader(const FString& HeaderName) OVERRIDE;
-	virtual TArray<FString> GetAllHeaders() OVERRIDE;	
-	virtual FString GetContentType() OVERRIDE;
-	virtual int32 GetContentLength() OVERRIDE;
-	virtual const TArray<uint8>& GetContent() OVERRIDE;
+	virtual FString GetURL() override;
+	virtual FString GetURLParameter(const FString& ParameterName) override;
+	virtual FString GetHeader(const FString& HeaderName) override;
+	virtual TArray<FString> GetAllHeaders() override;	
+	virtual FString GetContentType() override;
+	virtual int32 GetContentLength() override;
+	virtual const TArray<uint8>& GetContent() override;
 	// End IHttpBase interface
 
 	// Begin IHttpRequest interface
-	virtual FString GetVerb() OVERRIDE;
-	virtual void SetVerb(const FString& InVerb) OVERRIDE;
-	virtual void SetURL(const FString& InURL) OVERRIDE;
-	virtual void SetContent(const TArray<uint8>& ContentPayload) OVERRIDE;
-	virtual void SetContentAsString(const FString& ContentString) OVERRIDE;
-	virtual void SetHeader(const FString& HeaderName, const FString& HeaderValue) OVERRIDE;
-	virtual bool ProcessRequest() OVERRIDE;
-	virtual FHttpRequestCompleteDelegate& OnProcessRequestComplete() OVERRIDE;
-	virtual FHttpRequestProgressDelegate& OnRequestProgress() OVERRIDE;
-	virtual void CancelRequest() OVERRIDE;
-	virtual EHttpRequestStatus::Type GetStatus() OVERRIDE;
-	virtual void Tick(float DeltaSeconds) OVERRIDE;
+	virtual FString GetVerb() override;
+	virtual void SetVerb(const FString& InVerb) override;
+	virtual void SetURL(const FString& InURL) override;
+	virtual void SetContent(const TArray<uint8>& ContentPayload) override;
+	virtual void SetContentAsString(const FString& ContentString) override;
+	virtual void SetHeader(const FString& HeaderName, const FString& HeaderValue) override;
+	virtual bool ProcessRequest() override;
+	virtual FHttpRequestCompleteDelegate& OnProcessRequestComplete() override;
+	virtual FHttpRequestProgressDelegate& OnRequestProgress() override;
+	virtual void CancelRequest() override;
+	virtual EHttpRequestStatus::Type GetStatus() override;
+	virtual void Tick(float DeltaSeconds) override;
 	// End IHttpRequest interface
 
 	/**
@@ -333,18 +333,18 @@ public:
 
 
 	// Begin IHttpBase interface
-	virtual FString GetURL() OVERRIDE;
-	virtual FString GetURLParameter(const FString& ParameterName) OVERRIDE;
-	virtual FString GetHeader(const FString& HeaderName) OVERRIDE;
-	virtual TArray<FString> GetAllHeaders() OVERRIDE;	
-	virtual FString GetContentType() OVERRIDE;
-	virtual int32 GetContentLength() OVERRIDE;
-	virtual const TArray<uint8>& GetContent() OVERRIDE;
+	virtual FString GetURL() override;
+	virtual FString GetURLParameter(const FString& ParameterName) override;
+	virtual FString GetHeader(const FString& HeaderName) override;
+	virtual TArray<FString> GetAllHeaders() override;	
+	virtual FString GetContentType() override;
+	virtual int32 GetContentLength() override;
+	virtual const TArray<uint8>& GetContent() override;
 	// End IHttpBase interface
 
 	// Begin IHttpResponse interface
-	virtual int32 GetResponseCode() OVERRIDE;
-	virtual FString GetContentAsString() OVERRIDE;
+	virtual int32 GetResponseCode() override;
+	virtual FString GetContentAsString() override;
 	// End IHttpResponse interface
 
 	/**
diff --git a/Engine/Source/Runtime/Online/HTTP/Private/Curl/CurlHttpManager.h b/Engine/Source/Runtime/Online/HTTP/Private/Curl/CurlHttpManager.h
--- a/Engine/Source/Runtime/Online/HTTP/Private/Curl/CurlHttpManager.h
+++ b/Engine/Source/Runtime/Online/HTTP/Private/Curl/CurlHttpManager.h
@@ -28,9 +28,9 @@ protected:
 public:
 
 	// Begin HttpManager interface
-	virtual void AddRequest(TSharedRef<class IHttpRequest> Request) OVERRIDE;
-	virtual void RemoveRequest(TSharedRef<class IHttpRequest> Request) OVERRIDE;
-	virtual bool Tick(float DeltaSeconds) OVERRIDE;
+	virtual void AddRequest(TSharedRef<class IHttpRequest> Request) override;
+	virtual void RemoveRequest(TSharedRef<class IHttpRequest> Request) override;
+	virtual bool Tick(float DeltaSeconds) override;
 	// End HttpManager interface
 
 	FCurlHttpManager();
diff --git a/Engine/Source/Runtime/Online/HTTP/Private/IOS/IOSHTTP.h b/Engine/Source/Runtime/Online/HTTP/Private/IOS/IOSHTTP.h
--- a/Engine/Source/Runtime/Online/HTTP/Private/IOS/IOSHTTP.h
+++ b/Engine/Source/Runtime/Online/HTTP/Private/IOS/IOSHTTP.h
@@ -22,28 +22,28 @@ public:
 
 
 	// Begin IHttpBase interface
-	virtual FString GetURL() OVERRIDE;
-	virtual FString GetURLParameter(const FString& ParameterName) OVERRIDE;
-	virtual FString GetHeader(const FString& HeaderName) OVERRIDE;
-	virtual TArray<FString> GetAllHeaders() OVERRIDE;	
-	virtual FString GetContentType() OVERRIDE;
-	virtual int32 GetContentLength() OVERRIDE;
-	virtual const TArray<uint8>& GetContent() OVERRIDE;
+	virtual FString GetURL() override;
+	virtual FString GetURLParameter(const FString& ParameterName) override;
+	virtual FString GetHeader(const FString& HeaderName) override;
+	virtual TArray<FString> GetAllHeaders() override;	
+	virtual FString GetContentType() override;
+	virtual int32 GetContentLength() override;
+	virtual const TArray<uint8>& GetContent() override;
 	// End IHttpBase interface
 
 	// Begin IHttpRequest interface
-	virtual FString GetVerb() OVERRIDE;
-	virtual void SetVerb(const FString& Verb) OVERRIDE;
-	virtual void SetURL(const FString& URL) OVERRIDE;
-	virtual void SetContent(const TArray<uint8>& ContentPayload) OVERRIDE;
-	virtual void SetContentAsString(const FString& ContentString) OVERRIDE;
-	virtual void SetHeader(const FString& HeaderName, const FString& HeaderValue) OVERRIDE;
-	virtual bool ProcessRequest() OVERRIDE;
-	virtual FHttpRequestCompleteDelegate& OnProcessRequestComplete() OVERRIDE;
-	virtual FHttpRequestProgressDelegate& OnRequestProgress() OVERRIDE;
-	virtual void CancelRequest() OVERRIDE;
-	virtual EHttpRequestStatus::Type GetStatus() OVERRIDE;
-	virtual void Tick(float DeltaSeconds) OVERRIDE;
+	virtual FString GetVerb() override;
+	virtual void SetVerb(const FString& Verb) override;
+	virtual void SetURL(const FString& URL) override;
+	virtual void SetContent(const TArray<uint8>& ContentPayload) override;
+	virtual void SetContentAsString(const FString& ContentString) override;
+	virtual void SetHeader(const FString& HeaderName, const FString& HeaderValue) override;
+	virtual bool ProcessRequest() override;
+	virtual FHttpRequestCompleteDelegate& OnProcessRequestComplete() override;
+	virtual FHttpRequestProgressDelegate& OnRequestProgress() override;
+	virtual void CancelRequest() override;
+	virtual EHttpRequestStatus::Type GetStatus() override;
+	virtual void Tick(float DeltaSeconds) override;
 	// End IHttpRequest interface
 
 
@@ -143,18 +143,18 @@ public:
 
 
 	// Begin IHttpBase interface
-	virtual FString GetURL() OVERRIDE;
-	virtual FString GetURLParameter(const FString& ParameterName) OVERRIDE;
-	virtual FString GetHeader(const FString& HeaderName) OVERRIDE;
-	virtual TArray<FString> GetAllHeaders() OVERRIDE;	
-	virtual FString GetContentType() OVERRIDE;
-	virtual int32 GetContentLength() OVERRIDE;
-	virtual const TArray<uint8>& GetContent() OVERRIDE;
+	virtual FString GetURL() override;
+	virtual FString GetURLParameter(const FString& ParameterName) override;
+	virtual FString GetHeader(const FString& HeaderName) override;
+	virtual TArray<FString> GetAllHeaders() override;	
+	virtual FString GetContentType() override;
+	virtual int32 GetContentLength() override;
+	virtual const TArray<uint8>& GetContent() override;
 	// End IHttpBase interface
 
 	// Begin IHttpResponse interface
-	virtual int32 GetResponseCode() OVERRIDE;
-	virtual FString GetContentAsString() OVERRIDE;
+	virtual int32 GetResponseCode() override;
+	virtual FString GetContentAsString() override;
 	// End IHttpResponse interface
 
 	NSHTTPURLResponse* GetResponseObj();
diff --git a/Engine/Source/Runtime/Online/HTTP/Private/Mac/MacHTTP.h b/Engine/Source/Runtime/Online/HTTP/Private/Mac/MacHTTP.h
--- a/Engine/Source/Runtime/Online/HTTP/Private/Mac/MacHTTP.h
+++ b/Engine/Source/Runtime/Online/HTTP/Private/Mac/MacHTTP.h
@@ -22,28 +22,28 @@ public:
 
 
 	// Begin IHttpBase interface
-	virtual FString GetURL() OVERRIDE;
-	virtual FString GetURLParameter(const FString& ParameterName) OVERRIDE;
-	virtual FString GetHeader(const FString& HeaderName) OVERRIDE;
-	virtual TArray<FString> GetAllHeaders() OVERRIDE;	
-	virtual FString GetContentType() OVERRIDE;
-	virtual int32 GetContentLength() OVERRIDE;
-	virtual const TArray<uint8>& GetContent() OVERRIDE;
+	virtual FString GetURL() override;
+	virtual FString GetURLParameter(const FString& ParameterName) override;
+	virtual FString GetHeader(const FString& HeaderName) override;
+	virtual TArray<FString> GetAllHeaders() override;	
+	virtual FString GetContentType() override;
+	virtual int32 GetContentLength() override;
+	virtual const TArray<uint8>& GetContent() override;
 	// End IHttpBase interface
 
 	// Begin IHttpRequest interface
-	virtual FString GetVerb() OVERRIDE;
-	virtual void SetVerb(const FString& Verb) OVERRIDE;
-	virtual void SetURL(const FString& URL) OVERRIDE;
-	virtual void SetContent(const TArray<uint8>& ContentPayload) OVERRIDE;
-	virtual void SetContentAsString(const FString& ContentString) OVERRIDE;
-	virtual void SetHeader(const FString& HeaderName, const FString& HeaderValue) OVERRIDE;
-	virtual bool ProcessRequest() OVERRIDE;
-	virtual FHttpRequestCompleteDelegate& OnProcessRequestComplete() OVERRIDE;
-	virtual FHttpRequestProgressDelegate& OnRequestProgress() OVERRIDE;
-	virtual void CancelRequest() OVERRIDE;
-	virtual EHttpRequestStatus::Type GetStatus() OVERRIDE;
-	virtual void Tick(float DeltaSeconds) OVERRIDE;
+	virtual FString GetVerb() override;
+	virtual void SetVerb(const FString& Verb) override;
+	virtual void SetURL(const FString& URL) override;
+	virtual void SetContent(const TArray<uint8>& ContentPayload) override;
+	virtual void SetContentAsString(const FString& ContentString) override;
+	virtual void SetHeader(const FString& HeaderName, const FString& HeaderValue) override;
+	virtual bool ProcessRequest() override;
+	virtual FHttpRequestCompleteDelegate& OnProcessRequestComplete() override;
+	virtual FHttpRequestProgressDelegate& OnRequestProgress() override;
+	virtual void CancelRequest() override;
+	virtual EHttpRequestStatus::Type GetStatus() override;
+	virtual void Tick(float DeltaSeconds) override;
 	// End IHttpRequest interface
 
 
@@ -146,18 +146,18 @@ public:
 
 
 	// Begin IHttpBase interface
-	virtual FString GetURL() OVERRIDE;
-	virtual FString GetURLParameter(const FString& ParameterName) OVERRIDE;
-	virtual FString GetHeader(const FString& HeaderName) OVERRIDE;
-	virtual TArray<FString> GetAllHeaders() OVERRIDE;	
-	virtual FString GetContentType() OVERRIDE;
-	virtual int32 GetContentLength() OVERRIDE;
-	virtual const TArray<uint8>& GetContent() OVERRIDE;
+	virtual FString GetURL() override;
+	virtual FString GetURLParameter(const FString& ParameterName) override;
+	virtual FString GetHeader(const FString& HeaderName) override;
+	virtual TArray<FString> GetAllHeaders() override;	
+	virtual FString GetContentType() override;
+	virtual int32 GetContentLength() override;
+	virtual const TArray<uint8>& GetContent() override;
 	// End IHttpBase interface
 
 	// Begin IHttpResponse interface
-	virtual int32 GetResponseCode() OVERRIDE;
-	virtual FString GetContentAsString() OVERRIDE;
+	virtual int32 GetResponseCode() override;
+	virtual FString GetContentAsString() override;
 	// End IHttpResponse interface
 
 	NSHTTPURLResponse* GetResponseObj();
diff --git a/Engine/Source/Runtime/Online/HTTP/Public/HttpManager.h b/Engine/Source/Runtime/Online/HTTP/Public/HttpManager.h
--- a/Engine/Source/Runtime/Online/HTTP/Public/HttpManager.h
+++ b/Engine/Source/Runtime/Online/HTTP/Public/HttpManager.h
@@ -56,7 +56,7 @@ public:
 	 *
 	 * @return false if no longer needs ticking
 	 */
-	bool Tick(float DeltaSeconds) OVERRIDE;
+	bool Tick(float DeltaSeconds) override;
 
 	/**
 	 * List all of the Http requests currently being processed
diff --git a/Engine/Source/Runtime/Online/HTTP/Public/HttpModule.h b/Engine/Source/Runtime/Online/HTTP/Public/HttpModule.h
--- a/Engine/Source/Runtime/Online/HTTP/Public/HttpModule.h
+++ b/Engine/Source/Runtime/Online/HTTP/Public/HttpModule.h
@@ -26,7 +26,7 @@ public:
 	 *
 	 * @return true if the handler consumed the input, false to continue searching handlers
 	 */
-	virtual bool Exec(UWorld* InWorld, const TCHAR* Cmd, FOutputDevice& Ar) OVERRIDE;
+	virtual bool Exec(UWorld* InWorld, const TCHAR* Cmd, FOutputDevice& Ar) override;
 
 	/** 
 	 * Exec command handlers
@@ -142,13 +142,13 @@ private:
 	 * Called when Http module is loaded
 	 * Initialize platform specific parts of Http handling
 	 */
-	virtual void StartupModule() OVERRIDE;
+	virtual void StartupModule() override;
 	
 	/**
 	 * Called when Http module is unloaded
 	 * Shutdown platform specific parts of Http handling
 	 */
-	virtual void ShutdownModule() OVERRIDE;
+	virtual void ShutdownModule() override;
 
 
 	/** Keeps track of Http requests while they are being processed */
diff --git a/Engine/Source/Runtime/Online/IOS/OnlineSubsystemIOS/Private/OnlineAchievementsInterfaceIOS.h b/Engine/Source/Runtime/Online/IOS/OnlineSubsystemIOS/Private/OnlineAchievementsInterfaceIOS.h
--- a/Engine/Source/Runtime/Online/IOS/OnlineSubsystemIOS/Private/OnlineAchievementsInterfaceIOS.h
+++ b/Engine/Source/Runtime/Online/IOS/OnlineSubsystemIOS/Private/OnlineAchievementsInterfaceIOS.h
@@ -27,14 +27,14 @@ private:
 public:
 
 	// Begin IOnlineAchievements interface
-	virtual void WriteAchievements(const FUniqueNetId& PlayerId, FOnlineAchievementsWriteRef& WriteObject, const FOnAchievementsWrittenDelegate& Delegate = FOnAchievementsWrittenDelegate()) OVERRIDE;
-	virtual void QueryAchievements(const FUniqueNetId & PlayerId, const FOnQueryAchievementsCompleteDelegate & Delegate = FOnQueryAchievementsCompleteDelegate()) OVERRIDE;
-	virtual void QueryAchievementDescriptions( const FUniqueNetId& PlayerId, const FOnQueryAchievementsCompleteDelegate & Delegate = FOnQueryAchievementsCompleteDelegate() ) OVERRIDE;
-	virtual EOnlineCachedResult::Type GetCachedAchievement(const FUniqueNetId& PlayerId, const FString& AchievementId, FOnlineAchievement& OutAchievement) OVERRIDE;
-	virtual EOnlineCachedResult::Type GetCachedAchievements(const FUniqueNetId& PlayerId, TArray<FOnlineAchievement> & OutAchievements) OVERRIDE;
-	virtual EOnlineCachedResult::Type GetCachedAchievementDescription(const FString& AchievementId, FOnlineAchievementDesc& OutAchievementDesc) OVERRIDE;
+	virtual void WriteAchievements(const FUniqueNetId& PlayerId, FOnlineAchievementsWriteRef& WriteObject, const FOnAchievementsWrittenDelegate& Delegate = FOnAchievementsWrittenDelegate()) override;
+	virtual void QueryAchievements(const FUniqueNetId & PlayerId, const FOnQueryAchievementsCompleteDelegate & Delegate = FOnQueryAchievementsCompleteDelegate()) override;
+	virtual void QueryAchievementDescriptions( const FUniqueNetId& PlayerId, const FOnQueryAchievementsCompleteDelegate & Delegate = FOnQueryAchievementsCompleteDelegate() ) override;
+	virtual EOnlineCachedResult::Type GetCachedAchievement(const FUniqueNetId& PlayerId, const FString& AchievementId, FOnlineAchievement& OutAchievement) override;
+	virtual EOnlineCachedResult::Type GetCachedAchievements(const FUniqueNetId& PlayerId, TArray<FOnlineAchievement> & OutAchievements) override;
+	virtual EOnlineCachedResult::Type GetCachedAchievementDescription(const FString& AchievementId, FOnlineAchievementDesc& OutAchievementDesc) override;
 #if !UE_BUILD_SHIPPING
-	virtual bool ResetAchievements( const FUniqueNetId& PlayerId ) OVERRIDE;
+	virtual bool ResetAchievements( const FUniqueNetId& PlayerId ) override;
 #endif // !UE_BUILD_SHIPPING
 	// End IOnlineAchievements interface
 
diff --git a/Engine/Source/Runtime/Online/IOS/OnlineSubsystemIOS/Private/OnlineExternalUIInterfaceIOS.h b/Engine/Source/Runtime/Online/IOS/OnlineSubsystemIOS/Private/OnlineExternalUIInterfaceIOS.h
--- a/Engine/Source/Runtime/Online/IOS/OnlineSubsystemIOS/Private/OnlineExternalUIInterfaceIOS.h
+++ b/Engine/Source/Runtime/Online/IOS/OnlineSubsystemIOS/Private/OnlineExternalUIInterfaceIOS.h
@@ -14,13 +14,13 @@ PACKAGE_SCOPE:
 public:
 
 	// Begin IOnlineExternalUI interface
-	virtual bool ShowLoginUI(const int ControllerIndex, bool bShowOnlineOnly, const FOnLoginUIClosedDelegate& Delegate) OVERRIDE;
-	virtual bool ShowFriendsUI(int32 LocalUserNum) OVERRIDE;
-	virtual bool ShowInviteUI(int32 LocalUserNum) OVERRIDE;
-	virtual bool ShowAchievementsUI(int32 LocalUserNum) OVERRIDE;
-	virtual bool ShowLeaderboardUI(const FString& LeaderboardName) OVERRIDE;
-	virtual bool ShowWebURL(const FString& WebURL) OVERRIDE;
-	virtual bool ShowProfileUI(const FUniqueNetId& Requestor, const FUniqueNetId& Requestee, const FOnProfileUIClosedDelegate& Delegate) OVERRIDE;
+	virtual bool ShowLoginUI(const int ControllerIndex, bool bShowOnlineOnly, const FOnLoginUIClosedDelegate& Delegate) override;
+	virtual bool ShowFriendsUI(int32 LocalUserNum) override;
+	virtual bool ShowInviteUI(int32 LocalUserNum) override;
+	virtual bool ShowAchievementsUI(int32 LocalUserNum) override;
+	virtual bool ShowLeaderboardUI(const FString& LeaderboardName) override;
+	virtual bool ShowWebURL(const FString& WebURL) override;
+	virtual bool ShowProfileUI(const FUniqueNetId& Requestor, const FUniqueNetId& Requestee, const FOnProfileUIClosedDelegate& Delegate) override;
 	// End IOnlineExternalUI interface
 };
 
diff --git a/Engine/Source/Runtime/Online/IOS/OnlineSubsystemIOS/Private/OnlineFriendsInterfaceIOS.h b/Engine/Source/Runtime/Online/IOS/OnlineSubsystemIOS/Private/OnlineFriendsInterfaceIOS.h
--- a/Engine/Source/Runtime/Online/IOS/OnlineSubsystemIOS/Private/OnlineFriendsInterfaceIOS.h
+++ b/Engine/Source/Runtime/Online/IOS/OnlineSubsystemIOS/Private/OnlineFriendsInterfaceIOS.h
@@ -18,15 +18,15 @@ public:
 
 	// FOnlineUser
 
-	virtual TSharedRef<FUniqueNetId> GetUserId() const OVERRIDE;
-	virtual FString GetRealName() const OVERRIDE;
-	virtual FString GetDisplayName() const OVERRIDE;
-	virtual bool GetUserAttribute(const FString& AttrName, FString& OutAttrValue) const OVERRIDE;
+	virtual TSharedRef<FUniqueNetId> GetUserId() const override;
+	virtual FString GetRealName() const override;
+	virtual FString GetDisplayName() const override;
+	virtual bool GetUserAttribute(const FString& AttrName, FString& OutAttrValue) const override;
 
 	// FOnlineFriend
 	
-	virtual EInviteStatus::Type GetInviteStatus() const OVERRIDE;
-	virtual const FOnlineUserPresence& GetPresence() const OVERRIDE;
+	virtual EInviteStatus::Type GetInviteStatus() const override;
+	virtual const FOnlineUserPresence& GetPresence() const override;
 
 	// FOnlineFriendMcp
 
@@ -88,15 +88,15 @@ public:
 
 	// IOnlineFriends
 
-	virtual bool ReadFriendsList(int32 LocalUserNum, const FString& ListName) OVERRIDE;
-	virtual bool DeleteFriendsList(int32 LocalUserNum, const FString& ListName) OVERRIDE;
-	virtual bool SendInvite(int32 LocalUserNum, const FUniqueNetId& FriendId, const FString& ListName) OVERRIDE;
-	virtual bool AcceptInvite(int32 LocalUserNum, const FUniqueNetId& FriendId, const FString& ListName) OVERRIDE;
- 	virtual bool RejectInvite(int32 LocalUserNum, const FUniqueNetId& FriendId, const FString& ListName) OVERRIDE;
- 	virtual bool DeleteFriend(int32 LocalUserNum, const FUniqueNetId& FriendId, const FString& ListName) OVERRIDE;
-	virtual bool GetFriendsList(int32 LocalUserNum, const FString& ListName, TArray< TSharedRef<FOnlineFriend> >& OutFriends) OVERRIDE;
-	virtual TSharedPtr<FOnlineFriend> GetFriend(int32 LocalUserNum, const FUniqueNetId& FriendId, const FString& ListName) OVERRIDE;
-	virtual bool IsFriend(int32 LocalUserNum, const FUniqueNetId& FriendId, const FString& ListName) OVERRIDE;
+	virtual bool ReadFriendsList(int32 LocalUserNum, const FString& ListName) override;
+	virtual bool DeleteFriendsList(int32 LocalUserNum, const FString& ListName) override;
+	virtual bool SendInvite(int32 LocalUserNum, const FUniqueNetId& FriendId, const FString& ListName) override;
+	virtual bool AcceptInvite(int32 LocalUserNum, const FUniqueNetId& FriendId, const FString& ListName) override;
+ 	virtual bool RejectInvite(int32 LocalUserNum, const FUniqueNetId& FriendId, const FString& ListName) override;
+ 	virtual bool DeleteFriend(int32 LocalUserNum, const FUniqueNetId& FriendId, const FString& ListName) override;
+	virtual bool GetFriendsList(int32 LocalUserNum, const FString& ListName, TArray< TSharedRef<FOnlineFriend> >& OutFriends) override;
+	virtual TSharedPtr<FOnlineFriend> GetFriend(int32 LocalUserNum, const FUniqueNetId& FriendId, const FString& ListName) override;
+	virtual bool IsFriend(int32 LocalUserNum, const FUniqueNetId& FriendId, const FString& ListName) override;
 
 	// FOnlineFriendsIOS
 
diff --git a/Engine/Source/Runtime/Online/IOS/OnlineSubsystemIOS/Private/OnlineIdentityInterfaceIOS.h b/Engine/Source/Runtime/Online/IOS/OnlineSubsystemIOS/Private/OnlineIdentityInterfaceIOS.h
--- a/Engine/Source/Runtime/Online/IOS/OnlineSubsystemIOS/Private/OnlineIdentityInterfaceIOS.h
+++ b/Engine/Source/Runtime/Online/IOS/OnlineSubsystemIOS/Private/OnlineIdentityInterfaceIOS.h
@@ -28,17 +28,17 @@ PACKAGE_SCOPE:
 public:
 
 	// Begin IOnlineIdentity interface
-	virtual bool Login(int32 LocalUserNum, const FOnlineAccountCredentials& AccountCredentials) OVERRIDE;
-	virtual bool Logout(int32 LocalUserNum) OVERRIDE;
-	virtual bool AutoLogin(int32 LocalUserNum) OVERRIDE;
-	virtual TSharedPtr<FUserOnlineAccount> GetUserAccount(const FUniqueNetId& UserId) const OVERRIDE;
-	virtual TArray<TSharedPtr<FUserOnlineAccount> > GetAllUserAccounts() const OVERRIDE;
-	virtual TSharedPtr<FUniqueNetId> GetUniquePlayerId(int32 LocalUserNum) const OVERRIDE;
-	virtual TSharedPtr<FUniqueNetId> CreateUniquePlayerId(uint8* Bytes, int32 Size) OVERRIDE;
-	virtual TSharedPtr<FUniqueNetId> CreateUniquePlayerId(const FString& Str) OVERRIDE;
-	virtual ELoginStatus::Type GetLoginStatus(int32 LocalUserNum) const OVERRIDE;
-	virtual FString GetPlayerNickname(int32 LocalUserNum) const OVERRIDE;
-	virtual FString GetAuthToken(int32 LocalUserNum) const OVERRIDE;
+	virtual bool Login(int32 LocalUserNum, const FOnlineAccountCredentials& AccountCredentials) override;
+	virtual bool Logout(int32 LocalUserNum) override;
+	virtual bool AutoLogin(int32 LocalUserNum) override;
+	virtual TSharedPtr<FUserOnlineAccount> GetUserAccount(const FUniqueNetId& UserId) const override;
+	virtual TArray<TSharedPtr<FUserOnlineAccount> > GetAllUserAccounts() const override;
+	virtual TSharedPtr<FUniqueNetId> GetUniquePlayerId(int32 LocalUserNum) const override;
+	virtual TSharedPtr<FUniqueNetId> CreateUniquePlayerId(uint8* Bytes, int32 Size) override;
+	virtual TSharedPtr<FUniqueNetId> CreateUniquePlayerId(const FString& Str) override;
+	virtual ELoginStatus::Type GetLoginStatus(int32 LocalUserNum) const override;
+	virtual FString GetPlayerNickname(int32 LocalUserNum) const override;
+	virtual FString GetAuthToken(int32 LocalUserNum) const override;
 	// End IOnlineIdentity interface
 
 
diff --git a/Engine/Source/Runtime/Online/IOS/OnlineSubsystemIOS/Private/OnlineLeaderboardsInterfaceIOS.h b/Engine/Source/Runtime/Online/IOS/OnlineSubsystemIOS/Private/OnlineLeaderboardsInterfaceIOS.h
--- a/Engine/Source/Runtime/Online/IOS/OnlineSubsystemIOS/Private/OnlineLeaderboardsInterfaceIOS.h
+++ b/Engine/Source/Runtime/Online/IOS/OnlineSubsystemIOS/Private/OnlineLeaderboardsInterfaceIOS.h
@@ -23,12 +23,12 @@ public:
 	virtual ~FOnlineLeaderboardsIOS();
 
 	// Begin IOnlineLeaderboards interface
-	virtual bool ReadLeaderboards(const TArray< TSharedRef<FUniqueNetId> >& Players, FOnlineLeaderboardReadRef& ReadObject) OVERRIDE;
-	virtual bool ReadLeaderboardsForFriends(int32 LocalUserNum, FOnlineLeaderboardReadRef& ReadObject) OVERRIDE;
-	virtual void FreeStats(FOnlineLeaderboardRead& ReadObject) OVERRIDE;
-	virtual bool WriteLeaderboards(const FName& SessionName, const FUniqueNetId& Player, FOnlineLeaderboardWrite& WriteObject) OVERRIDE;
-	virtual bool FlushLeaderboards(const FName& SessionName) OVERRIDE;
-	virtual bool WriteOnlinePlayerRatings(const FName& SessionName, int32 LeaderboardId, const TArray<FOnlinePlayerScore>& PlayerScores) OVERRIDE;
+	virtual bool ReadLeaderboards(const TArray< TSharedRef<FUniqueNetId> >& Players, FOnlineLeaderboardReadRef& ReadObject) override;
+	virtual bool ReadLeaderboardsForFriends(int32 LocalUserNum, FOnlineLeaderboardReadRef& ReadObject) override;
+	virtual void FreeStats(FOnlineLeaderboardRead& ReadObject) override;
+	virtual bool WriteLeaderboards(const FName& SessionName, const FUniqueNetId& Player, FOnlineLeaderboardWrite& WriteObject) override;
+	virtual bool FlushLeaderboards(const FName& SessionName) override;
+	virtual bool WriteOnlinePlayerRatings(const FName& SessionName, int32 LeaderboardId, const TArray<FOnlinePlayerScore>& PlayerScores) override;
 	// End IOnlineLeaderboards interface
 
 };
