Build fix. Header include order is important and must not redeclare log categories.

diff --git a/Engine/Source/Runtime/NetworkFile/Private/NetworkFilePrivatePCH.h b/Engine/Source/Runtime/NetworkFile/Private/NetworkFilePrivatePCH.h
--- a/Engine/Source/Runtime/NetworkFile/Private/NetworkFilePrivatePCH.h
+++ b/Engine/Source/Runtime/NetworkFile/Private/NetworkFilePrivatePCH.h
@@ -2,7 +2,7 @@
 
 #pragma once
 
-#include "NetworkMessage.h"
-
 #include "Core.h"
-#include "ModuleManager.h"
\ No newline at end of file
+#include "ModuleManager.h"
+
+#include "NetworkMessage.h"
\ No newline at end of file
diff --git a/Engine/Source/Runtime/NetworkFile/Private/TCPTransport.cpp b/Engine/Source/Runtime/NetworkFile/Private/TCPTransport.cpp
--- a/Engine/Source/Runtime/NetworkFile/Private/TCPTransport.cpp
+++ b/Engine/Source/Runtime/NetworkFile/Private/TCPTransport.cpp
@@ -4,8 +4,6 @@
 #include "TCPTransport.h"
 #include "MultichannelTCP.h"
 
-DECLARE_LOG_CATEGORY_EXTERN(LogNetworkPlatformFile, Log, All);
-
 FTCPTransport::FTCPTransport()
 	:FileSocket(NULL)
 	,MCSocket(NULL)
