Another fix for CL# 2103382. This one is for logging.

diff --git a/Engine/Source/Runtime/NetworkFile/Private/TCPTransport.cpp b/Engine/Source/Runtime/NetworkFile/Private/TCPTransport.cpp
--- a/Engine/Source/Runtime/NetworkFile/Private/TCPTransport.cpp
+++ b/Engine/Source/Runtime/NetworkFile/Private/TCPTransport.cpp
@@ -3,6 +3,7 @@
 #include "NetworkFilePrivatePCH.h"
 #include "TCPTransport.h"
 #include "MultichannelTCP.h"
+#include "NetworkPlatformFile.h"
 
 FTCPTransport::FTCPTransport()
 	:FileSocket(NULL)
