fix for crash when running a non-streaming cook on the fly server
#ue4

diff --git a/Engine/Source/Runtime/NetworkFile/Private/NetworkPlatformFile.cpp b/Engine/Source/Runtime/NetworkFile/Private/NetworkPlatformFile.cpp
--- a/Engine/Source/Runtime/NetworkFile/Private/NetworkPlatformFile.cpp
+++ b/Engine/Source/Runtime/NetworkFile/Private/NetworkPlatformFile.cpp
@@ -88,9 +88,9 @@ bool FNetworkPlatformFile::InitializeInternal(IPlatformFile* Inner, const TCHAR*
 	LocalDirectories.Add(FPaths::GameSavedDir() / TEXT("Logs"));
 	LocalDirectories.Add(FPaths::GameSavedDir() / TEXT("Sandboxes"));
 
-	FArrayWriter In; 
+	FNetworkFileArchive Payload(NFS_Messages::Heartbeat); 
 	FArrayReader Out;
-	if (!SendPayloadAndReceiveResponse(In,Out))
+	if (!SendPayloadAndReceiveResponse(Payload,Out))
 		bIsUsable = true; 
 
 
