#UE4

       Temporary fix async read/write thread hang for NFS caching.

diff --git a/Engine/Source/Runtime/NetworkFile/Private/NetworkPlatformFile.cpp b/Engine/Source/Runtime/NetworkFile/Private/NetworkPlatformFile.cpp
--- a/Engine/Source/Runtime/NetworkFile/Private/NetworkPlatformFile.cpp
+++ b/Engine/Source/Runtime/NetworkFile/Private/NetworkPlatformFile.cpp
@@ -105,8 +105,6 @@ bool FNetworkPlatformFile::InitializeInternal(IPlatformFile* Inner, const TCHAR*
 
 bool FNetworkPlatformFile::SendPayloadAndReceiveResponse(TArray<uint8>& In, TArray<uint8>& Out)
 {
-	delete FinishedAsyncReadUnsolicitedFiles; // wait here for any async unsolicited files to finish being read from the network before we start to read/write again.
-	FinishedAsyncReadUnsolicitedFiles = NULL;
 	return Transport->SendPayloadAndReceiveResponse( In, Out );
 }
 
@@ -798,8 +796,8 @@ void AsyncWriteFile(FArchive* Archive, const FString& Filename, FDateTime Server
 	(new FAutoDeleteAsyncTask<FAsyncNetworkWriteWorker>(*Filename, Archive, ServerTimeStamp, &InnerPlatformFile, Event))->StartBackgroundTask();
 }
 
-void AsyncReadUnsolicitedFile(int32 InNumUnsolictedFiles, FNetworkPlatformFile& InNetworkFile, FScopedEvent* InEvent, IPlatformFile& InInnerPlatformFile, 
-							  FScopedEvent* InAllDoneEvent, FString& InServerEngineDir, FString& InServerGameDir)
+void ReadUnsolicitedFile(int32 InNumUnsolictedFiles, FNetworkPlatformFile& InNetworkFile, FScopedEvent* InEvent, IPlatformFile& InInnerPlatformFile, 
+							  FScopedEvent* InAllDoneEvent, FString& InServerEngineDir, FString& InServerGameDir, bool IsAsync )
 {
 	class FAsyncReadUnsolicitedFile : public FNonAbandonableTask
 	{
@@ -834,8 +832,8 @@ void AsyncReadUnsolicitedFile(int32 InNumUnsolictedFiles, FNetworkPlatformFile&
 			for (int32 Index = 0; Index < NumUnsolictedFiles; Index++)
 			{
 				FArrayReader* UnsolictedResponse = new FArrayReader;
-				TArray<uint8> PayLoad; //empty payload. 
-				if (!NetworkFile.SendPayloadAndReceiveResponse(PayLoad,*UnsolictedResponse))
+				FNetworkFileArchive Payload(NFS_Messages::Heartbeat); 
+				if (!NetworkFile.SendPayloadAndReceiveResponse(Payload,*UnsolictedResponse))
 				{
 					UE_LOG(LogNetworkPlatformFile, Fatal, TEXT("Receive failure!"));
 					return;
@@ -849,7 +847,7 @@ void AsyncReadUnsolicitedFile(int32 InNumUnsolictedFiles, FNetworkPlatformFile&
 				*UnsolictedResponse << UnsolictedServerTimeStamp;
 
 				// write the file by pulling out of the FArrayReader
-				AsyncWriteFile(UnsolictedResponse, UnsolictedReplyFile, UnsolictedServerTimeStamp, InnerPlatformFile, AllDoneEvent);
+				SyncWriteFile(UnsolictedResponse, UnsolictedReplyFile, UnsolictedServerTimeStamp, InnerPlatformFile);
 			}
 			Event->Trigger();
 		}
@@ -861,7 +859,11 @@ void AsyncReadUnsolicitedFile(int32 InNumUnsolictedFiles, FNetworkPlatformFile&
 			return TEXT("FAsyncReadUnsolicitedFile");
 		}
 	};
-	(new FAutoDeleteAsyncTask<FAsyncReadUnsolicitedFile>(InNumUnsolictedFiles, &InNetworkFile, InEvent, &InInnerPlatformFile, InAllDoneEvent, InServerEngineDir, InServerGameDir))->StartBackgroundTask();
+
+	if ( IsAsync )
+		(new FAutoDeleteAsyncTask<FAsyncReadUnsolicitedFile>(InNumUnsolictedFiles, &InNetworkFile, InEvent, &InInnerPlatformFile, InAllDoneEvent, InServerEngineDir, InServerGameDir))->StartBackgroundTask();
+	else 
+		(new FAutoDeleteAsyncTask<FAsyncReadUnsolicitedFile>(InNumUnsolictedFiles, &InNetworkFile, InEvent, &InInnerPlatformFile, InAllDoneEvent, InServerEngineDir, InServerGameDir))->StartSynchronousTask();
 }
 
 /**
@@ -978,7 +980,7 @@ void FNetworkPlatformFile::EnsureFileIsLocal(const FString& Filename)
 	{
 		FinishedAsyncReadUnsolicitedFiles = new FScopedEvent;
 		FinishedAsyncWriteUnsolicitedFiles = new FScopedEvent;
-		AsyncReadUnsolicitedFile(NumUnsolictedFiles, *this, FinishedAsyncReadUnsolicitedFiles, *InnerPlatformFile, FinishedAsyncWriteUnsolicitedFiles, ServerEngineDir, ServerGameDir);
+		ReadUnsolicitedFile(NumUnsolictedFiles, *this, FinishedAsyncReadUnsolicitedFiles, *InnerPlatformFile, FinishedAsyncWriteUnsolicitedFiles, ServerEngineDir, ServerGameDir);
 	}
 
 	ThisTime = 1000.0f * float(FPlatformTime::Seconds() - StartTime);
diff --git a/Engine/Source/Runtime/NetworkFile/Public/NetworkPlatformFile.h b/Engine/Source/Runtime/NetworkFile/Public/NetworkPlatformFile.h
--- a/Engine/Source/Runtime/NetworkFile/Public/NetworkPlatformFile.h
+++ b/Engine/Source/Runtime/NetworkFile/Public/NetworkPlatformFile.h
@@ -13,8 +13,8 @@ DECLARE_LOG_CATEGORY_EXTERN(LogNetworkPlatformFile, Log, All);
 class NETWORKFILE_API FNetworkPlatformFile : public IPlatformFile
 {
 	friend class FAsyncFileSync;
-	friend void AsyncReadUnsolicitedFile(int32 InNumUnsolictedFiles, FNetworkPlatformFile& InNetworkFile, FScopedEvent* InEvent, IPlatformFile& InInnerPlatformFile, 
-		FScopedEvent* InAllDoneEvent, FString& InServerEngineDir, FString& InServerGameDir);
+	friend void ReadUnsolicitedFile(int32 InNumUnsolictedFiles, FNetworkPlatformFile& InNetworkFile, FScopedEvent* InEvent, IPlatformFile& InInnerPlatformFile, 
+		FScopedEvent* InAllDoneEvent, FString& InServerEngineDir, FString& InServerGameDir, bool IsAsync = false);
 
 	/**
 	 * Initialize network platform file give the specified host IP
