#UE4

     Remove async read/write path for now -  the server never supported really supported full duplex so the extra complexity is probably not worth it. Need to rethink how async path should work.

diff --git a/Engine/Source/Runtime/NetworkFile/Private/NetworkPlatformFile.cpp b/Engine/Source/Runtime/NetworkFile/Private/NetworkPlatformFile.cpp
--- a/Engine/Source/Runtime/NetworkFile/Private/NetworkPlatformFile.cpp
+++ b/Engine/Source/Runtime/NetworkFile/Private/NetworkPlatformFile.cpp
@@ -25,8 +25,6 @@ FNetworkPlatformFile::FNetworkPlatformFile()
 	, InnerPlatformFile(NULL)
 	, bIsUsable(false)
 	, FileServerPort(DEFAULT_FILE_SERVING_PORT)
-	, FinishedAsyncReadUnsolicitedFiles(NULL)
-	, FinishedAsyncWriteUnsolicitedFiles(NULL)
 	, Transport(NULL)
 {
 #if !USE_HTTP_FOR_NFS
@@ -289,10 +287,6 @@ FNetworkPlatformFile::~FNetworkPlatformFile()
 	if (!GIsRequestingExit) // the socket subsystem is probably already gone, so it will crash if we clean up
 	{
 		delete Transport; // close our sockets.
-		delete FinishedAsyncReadUnsolicitedFiles; // wait here for any async unsolicited files to finish reading being read from the network 
-		FinishedAsyncReadUnsolicitedFiles = NULL;
-		delete FinishedAsyncWriteUnsolicitedFiles; // wait here for any async unsolicited files to finish writing 
-		FinishedAsyncWriteUnsolicitedFiles = NULL;
 	}
 }
 
@@ -670,8 +664,6 @@ bool FNetworkPlatformFile::SendMessageToServer(const TCHAR* Message, IPlatformFi
 }
 
 
-//@todo, this should really be a member of FNetworkPlatformFile
-static FThreadSafeCounter OutstandingAsyncWrites;
 
 class FAsyncNetworkWriteWorker : public FNonAbandonableTask
 {
@@ -767,10 +759,7 @@ public:
 		}
 		if (Event)
 		{
-			if (!OutstandingAsyncWrites.Decrement())
-			{
 				Event->Trigger(); // last file, fire trigger
-			}
 		}
 	}
 	/** Give the name for external event viewers
@@ -791,44 +780,29 @@ void SyncWriteFile(FArchive* Archive, const FString& Filename, FDateTime ServerT
 	(new FAutoDeleteAsyncTask<FAsyncNetworkWriteWorker>(*Filename, Archive, ServerTimeStamp, &InnerPlatformFile, NullEvent))->StartSynchronousTask();
 }
 
-void AsyncWriteFile(FArchive* Archive, const FString& Filename, FDateTime ServerTimeStamp, IPlatformFile& InnerPlatformFile, FScopedEvent* Event = NULL)
-{
-	(new FAutoDeleteAsyncTask<FAsyncNetworkWriteWorker>(*Filename, Archive, ServerTimeStamp, &InnerPlatformFile, Event))->StartBackgroundTask();
-}
-
-void ReadUnsolicitedFile(int32 InNumUnsolictedFiles, FNetworkPlatformFile& InNetworkFile, FScopedEvent* InEvent, IPlatformFile& InInnerPlatformFile, 
-							  FScopedEvent* InAllDoneEvent, FString& InServerEngineDir, FString& InServerGameDir, bool IsAsync )
+void ReadUnsolicitedFile(int32 InNumUnsolictedFiles, FNetworkPlatformFile& InNetworkFile, IPlatformFile& InInnerPlatformFile, FString& InServerEngineDir, FString& InServerGameDir)
 {
 	class FAsyncReadUnsolicitedFile : public FNonAbandonableTask
 	{
 	public:
 		int32 NumUnsolictedFiles;
 		FNetworkPlatformFile& NetworkFile;
-		FScopedEvent* Event;
 		IPlatformFile& InnerPlatformFile;
-		FScopedEvent* AllDoneEvent;
 		FString ServerEngineDir;
 		FString ServerGameDir;
 
-		FAsyncReadUnsolicitedFile(int32 In_NumUnsolictedFiles, FNetworkPlatformFile* In_NetworkFile, FScopedEvent* In_Event, IPlatformFile* In_InnerPlatformFile, 
-			FScopedEvent* In_AllDoneEvent, FString& In_ServerEngineDir, FString& In_ServerGameDir)
+		FAsyncReadUnsolicitedFile(int32 In_NumUnsolictedFiles, FNetworkPlatformFile* In_NetworkFile, IPlatformFile* In_InnerPlatformFile, FString& In_ServerEngineDir, FString& In_ServerGameDir)
 			: NumUnsolictedFiles(In_NumUnsolictedFiles)
 			, NetworkFile(*In_NetworkFile)
-			, Event(In_Event)
 			, InnerPlatformFile(*In_InnerPlatformFile)
-			, AllDoneEvent(In_AllDoneEvent)
 			, ServerEngineDir(In_ServerEngineDir)
 			, ServerGameDir(In_ServerGameDir)
 		{
-			check(Event);
-			check(AllDoneEvent);
 		}
 		
 		/** Write the file  */
 		void DoWork()
 		{
-			check(!OutstandingAsyncWrites.GetValue());
-			OutstandingAsyncWrites.Add(NumUnsolictedFiles);
 			for (int32 Index = 0; Index < NumUnsolictedFiles; Index++)
 			{
 				FArrayReader* UnsolictedResponse = new FArrayReader;
@@ -840,16 +814,18 @@ void ReadUnsolicitedFile(int32 InNumUnsolictedFiles, FNetworkPlatformFile& InNet
 				}
 				FString UnsolictedReplyFile;
 				*UnsolictedResponse << UnsolictedReplyFile;
-				FNetworkPlatformFile::ConvertServerFilenameToClientFilename(UnsolictedReplyFile, ServerEngineDir, ServerGameDir);
 
-				// get the server file timestamp
-				FDateTime UnsolictedServerTimeStamp;
-				*UnsolictedResponse << UnsolictedServerTimeStamp;
+				if (!UnsolictedReplyFile.IsEmpty())
+				{
+					FNetworkPlatformFile::ConvertServerFilenameToClientFilename(UnsolictedReplyFile, ServerEngineDir, ServerGameDir);
+					// get the server file timestamp
+					FDateTime UnsolictedServerTimeStamp;
+					*UnsolictedResponse << UnsolictedServerTimeStamp;
 
-				// write the file by pulling out of the FArrayReader
-				SyncWriteFile(UnsolictedResponse, UnsolictedReplyFile, UnsolictedServerTimeStamp, InnerPlatformFile);
+					// write the file by pulling out of the FArrayReader
+					SyncWriteFile(UnsolictedResponse, UnsolictedReplyFile, UnsolictedServerTimeStamp, InnerPlatformFile);
+				}
 			}
-			Event->Trigger();
 		}
 		/** Give the name for external event viewers
 		* @return	the name to display in external event viewers
@@ -859,11 +835,7 @@ void ReadUnsolicitedFile(int32 InNumUnsolictedFiles, FNetworkPlatformFile& InNet
 			return TEXT("FAsyncReadUnsolicitedFile");
 		}
 	};
-
-	if ( IsAsync )
-		(new FAutoDeleteAsyncTask<FAsyncReadUnsolicitedFile>(InNumUnsolictedFiles, &InNetworkFile, InEvent, &InInnerPlatformFile, InAllDoneEvent, InServerEngineDir, InServerGameDir))->StartBackgroundTask();
-	else 
-		(new FAutoDeleteAsyncTask<FAsyncReadUnsolicitedFile>(InNumUnsolictedFiles, &InNetworkFile, InEvent, &InInnerPlatformFile, InAllDoneEvent, InServerEngineDir, InServerGameDir))->StartSynchronousTask();
+	 (new FAutoDeleteAsyncTask<FAsyncReadUnsolicitedFile>(InNumUnsolictedFiles, &InNetworkFile, &InInnerPlatformFile, InServerEngineDir, InServerGameDir))->StartSynchronousTask();
 }
 
 /**
@@ -885,11 +857,6 @@ void FNetworkPlatformFile::EnsureFileIsLocal(const FString& Filename)
 		}
 	}
 
-	delete FinishedAsyncReadUnsolicitedFiles; // wait here for any async unsolicited files to finish reading being read from the network 
-	FinishedAsyncReadUnsolicitedFiles = NULL;
-	delete FinishedAsyncWriteUnsolicitedFiles; // wait here for any async unsolicited files to finish writing 
-	FinishedAsyncWriteUnsolicitedFiles = NULL;
-
 	FScopeLock ScopeLock(&SynchronizationObject);
 
 	ThisTime = 1000.0f * float(FPlatformTime::Seconds() - StartTime);
@@ -978,9 +945,7 @@ void FNetworkPlatformFile::EnsureFileIsLocal(const FString& Filename)
 
 	if (NumUnsolictedFiles)
 	{
-		FinishedAsyncReadUnsolicitedFiles = new FScopedEvent;
-		FinishedAsyncWriteUnsolicitedFiles = new FScopedEvent;
-		ReadUnsolicitedFile(NumUnsolictedFiles, *this, FinishedAsyncReadUnsolicitedFiles, *InnerPlatformFile, FinishedAsyncWriteUnsolicitedFiles, ServerEngineDir, ServerGameDir);
+		ReadUnsolicitedFile(NumUnsolictedFiles, *this, *InnerPlatformFile, ServerEngineDir, ServerGameDir);
 	}
 
 	ThisTime = 1000.0f * float(FPlatformTime::Seconds() - StartTime);
diff --git a/Engine/Source/Runtime/NetworkFile/Public/NetworkPlatformFile.h b/Engine/Source/Runtime/NetworkFile/Public/NetworkPlatformFile.h
--- a/Engine/Source/Runtime/NetworkFile/Public/NetworkPlatformFile.h
+++ b/Engine/Source/Runtime/NetworkFile/Public/NetworkPlatformFile.h
@@ -13,8 +13,7 @@ DECLARE_LOG_CATEGORY_EXTERN(LogNetworkPlatformFile, Log, All);
 class NETWORKFILE_API FNetworkPlatformFile : public IPlatformFile
 {
 	friend class FAsyncFileSync;
-	friend void ReadUnsolicitedFile(int32 InNumUnsolictedFiles, FNetworkPlatformFile& InNetworkFile, FScopedEvent* InEvent, IPlatformFile& InInnerPlatformFile, 
-		FScopedEvent* InAllDoneEvent, FString& InServerEngineDir, FString& InServerGameDir, bool IsAsync = false);
+	friend void ReadUnsolicitedFile(int32 InNumUnsolictedFiles, FNetworkPlatformFile& InNetworkFile, IPlatformFile& InInnerPlatformFile,  FString& InServerEngineDir, FString& InServerGameDir);
 
 	/**
 	 * Initialize network platform file give the specified host IP
@@ -200,9 +199,7 @@ protected:
 	bool				bIsUsable;
 	int32				FileServerPort;
 private:
-	FScopedEvent* FinishedAsyncReadUnsolicitedFiles;
-	FScopedEvent* FinishedAsyncWriteUnsolicitedFiles;
-
+	
     // Our network Transport. 
 	class ITransport* Transport; 
 };
