Fix for slow cook on the fly issue.
Reverted a bunch unsolicited file handling code back to the way it used to be.

diff --git a/Engine/Source/Runtime/NetworkFile/Private/HTTPTransport.h b/Engine/Source/Runtime/NetworkFile/Private/HTTPTransport.h
--- a/Engine/Source/Runtime/NetworkFile/Private/HTTPTransport.h
+++ b/Engine/Source/Runtime/NetworkFile/Private/HTTPTransport.h
@@ -15,8 +15,9 @@ public:
 	FHTTPTransport();
 
 	// ITransport Interface.
-	virtual bool Initialize(const TCHAR* HostIp);
-	virtual bool SendPayloadAndReceiveResponse(TArray<uint8>& In, TArray<uint8>& Out);
+	virtual bool Initialize(const TCHAR* HostIp) override;
+	virtual bool SendPayloadAndReceiveResponse(TArray<uint8>& In, TArray<uint8>& Out) override;
+	virtual bool ReceiveResponse(TArray<uint8> &Out) override { return false; } // unsupported for now :(
 
 private: 
 
diff --git a/Engine/Source/Runtime/NetworkFile/Private/ITransport.h b/Engine/Source/Runtime/NetworkFile/Private/ITransport.h
--- a/Engine/Source/Runtime/NetworkFile/Private/ITransport.h
+++ b/Engine/Source/Runtime/NetworkFile/Private/ITransport.h
@@ -16,5 +16,6 @@ public:
 
 	virtual bool Initialize(const TCHAR*) = 0; 
 	virtual bool SendPayloadAndReceiveResponse(TArray<uint8>& In, TArray<uint8>& Out) = 0; 
+	virtual bool ReceiveResponse(TArray<uint8> &Out) = 0;
 
 };
\ No newline at end of file
diff --git a/Engine/Source/Runtime/NetworkFile/Private/NetworkPlatformFile.cpp b/Engine/Source/Runtime/NetworkFile/Private/NetworkPlatformFile.cpp
--- a/Engine/Source/Runtime/NetworkFile/Private/NetworkPlatformFile.cpp
+++ b/Engine/Source/Runtime/NetworkFile/Private/NetworkPlatformFile.cpp
@@ -24,6 +24,8 @@ FNetworkPlatformFile::FNetworkPlatformFile()
 	: bHasLoadedDDCDirectories(false)
 	, InnerPlatformFile(NULL)
 	, bIsUsable(false)
+	, FinishedAsyncNetworkReadUnsolicitedFiles(NULL)
+	, FinishedAsyncWriteUnsolicitedFiles(NULL)
 	, Transport(NULL)
 {
 }
@@ -129,9 +131,21 @@ bool FNetworkPlatformFile::InitializeInternal(IPlatformFile* Inner, const TCHAR*
 
 bool FNetworkPlatformFile::SendPayloadAndReceiveResponse(TArray<uint8>& In, TArray<uint8>& Out)
 {
+	if ( FinishedAsyncNetworkReadUnsolicitedFiles )
+	{
+		delete FinishedAsyncNetworkReadUnsolicitedFiles;
+		FinishedAsyncNetworkReadUnsolicitedFiles = NULL;
+	}
+	
 	return Transport->SendPayloadAndReceiveResponse( In, Out );
 }
 
+bool FNetworkPlatformFile::ReceiveResponse(TArray<uint8> &Out )
+{
+	return Transport->ReceiveResponse( Out );
+}
+
+
 void FNetworkPlatformFile::InitializeAfterSetActive()
 {
 	double NetworkFileStartupTime = 0.0;
@@ -312,6 +326,17 @@ FNetworkPlatformFile::~FNetworkPlatformFile()
 {
 	if (!GIsRequestingExit) // the socket subsystem is probably already gone, so it will crash if we clean up
 	{
+		if ( FinishedAsyncNetworkReadUnsolicitedFiles )
+		{
+			delete FinishedAsyncNetworkReadUnsolicitedFiles; // wait here for any async unsolicited files to finish reading being read from the network 
+			FinishedAsyncNetworkReadUnsolicitedFiles = NULL;
+		}
+		if ( FinishedAsyncWriteUnsolicitedFiles )
+		{
+			delete FinishedAsyncWriteUnsolicitedFiles; // wait here for any async unsolicited files to finish writing 
+			FinishedAsyncWriteUnsolicitedFiles = NULL;
+		}
+		
 		delete Transport; // close our sockets.
 	}
 }
@@ -690,6 +715,7 @@ bool FNetworkPlatformFile::SendMessageToServer(const TCHAR* Message, IPlatformFi
 }
 
 
+static FThreadSafeCounter OutstandingAsyncWrites;
 
 class FAsyncNetworkWriteWorker : public FNonAbandonableTask
 {
@@ -785,7 +811,10 @@ public:
 		}
 		if (Event)
 		{
+			if (OutstandingAsyncWrites.Decrement() == 0)
+			{
 				Event->Trigger(); // last file, fire trigger
+			}
 		}
 	}
 	/** Give the name for external event viewers
@@ -806,7 +835,12 @@ void SyncWriteFile(FArchive* Archive, const FString& Filename, FDateTime ServerT
 	(new FAutoDeleteAsyncTask<FAsyncNetworkWriteWorker>(*Filename, Archive, ServerTimeStamp, &InnerPlatformFile, NullEvent))->StartSynchronousTask();
 }
 
-void ReadUnsolicitedFile(int32 InNumUnsolictedFiles, FNetworkPlatformFile& InNetworkFile, IPlatformFile& InInnerPlatformFile, FString& InServerEngineDir, FString& InServerGameDir)
+void AsyncWriteFile(FArchive* Archive, const FString& Filename, FDateTime ServerTimeStamp, IPlatformFile& InnerPlatformFile, FScopedEvent* Event = NULL)
+{
+	(new FAutoDeleteAsyncTask<FAsyncNetworkWriteWorker>(*Filename, Archive, ServerTimeStamp, &InnerPlatformFile, Event))->StartBackgroundTask();
+}
+
+void AsyncReadUnsolicitedFiles(int32 InNumUnsolictedFiles, FNetworkPlatformFile& InNetworkFile, IPlatformFile& InInnerPlatformFile, FString& InServerEngineDir, FString& InServerGameDir, FScopedEvent *InNetworkDoneEvent, FScopedEvent *InWritingDoneEvent)
 {
 	class FAsyncReadUnsolicitedFile : public FNonAbandonableTask
 	{
@@ -816,24 +850,28 @@ void ReadUnsolicitedFile(int32 InNumUnsolictedFiles, FNetworkPlatformFile& InNet
 		IPlatformFile& InnerPlatformFile;
 		FString ServerEngineDir;
 		FString ServerGameDir;
+		FScopedEvent* NetworkDoneEvent; // finished using the network
+		FScopedEvent* WritingDoneEvent; // finished writing the files to disk
 
-		FAsyncReadUnsolicitedFile(int32 In_NumUnsolictedFiles, FNetworkPlatformFile* In_NetworkFile, IPlatformFile* In_InnerPlatformFile, FString& In_ServerEngineDir, FString& In_ServerGameDir)
+		FAsyncReadUnsolicitedFile(int32 In_NumUnsolictedFiles, FNetworkPlatformFile* In_NetworkFile, IPlatformFile* In_InnerPlatformFile, FString& In_ServerEngineDir, FString& In_ServerGameDir, FScopedEvent *InNetworkDoneEvent, FScopedEvent *InWritingDoneEvent )
 			: NumUnsolictedFiles(In_NumUnsolictedFiles)
 			, NetworkFile(*In_NetworkFile)
 			, InnerPlatformFile(*In_InnerPlatformFile)
 			, ServerEngineDir(In_ServerEngineDir)
 			, ServerGameDir(In_ServerGameDir)
+			, NetworkDoneEvent(InNetworkDoneEvent)
+			, WritingDoneEvent(InWritingDoneEvent)
 		{
 		}
 		
 		/** Write the file  */
 		void DoWork()
 		{
+			OutstandingAsyncWrites.Add( NumUnsolictedFiles );
 			for (int32 Index = 0; Index < NumUnsolictedFiles; Index++)
 			{
 				FArrayReader* UnsolictedResponse = new FArrayReader;
-				FNetworkFileArchive Payload(NFS_Messages::Heartbeat); 
-				if (!NetworkFile.SendPayloadAndReceiveResponse(Payload,*UnsolictedResponse))
+				if (!NetworkFile.ReceiveResponse(*UnsolictedResponse))
 				{
 					UE_LOG(LogNetworkPlatformFile, Fatal, TEXT("Receive failure!"));
 					return;
@@ -849,9 +887,10 @@ void ReadUnsolicitedFile(int32 InNumUnsolictedFiles, FNetworkPlatformFile& InNet
 					*UnsolictedResponse << UnsolictedServerTimeStamp;
 
 					// write the file by pulling out of the FArrayReader
-					SyncWriteFile(UnsolictedResponse, UnsolictedReplyFile, UnsolictedServerTimeStamp, InnerPlatformFile);
+					AsyncWriteFile(UnsolictedResponse, UnsolictedReplyFile, UnsolictedServerTimeStamp, InnerPlatformFile, WritingDoneEvent);
 				}
 			}
+			NetworkDoneEvent->Trigger();
 		}
 		/** Give the name for external event viewers
 		* @return	the name to display in external event viewers
@@ -861,7 +900,8 @@ void ReadUnsolicitedFile(int32 InNumUnsolictedFiles, FNetworkPlatformFile& InNet
 			return TEXT("FAsyncReadUnsolicitedFile");
 		}
 	};
-	 (new FAutoDeleteAsyncTask<FAsyncReadUnsolicitedFile>(InNumUnsolictedFiles, &InNetworkFile, &InInnerPlatformFile, InServerEngineDir, InServerGameDir))->StartSynchronousTask();
+
+	(new FAutoDeleteAsyncTask<FAsyncReadUnsolicitedFile>(InNumUnsolictedFiles, &InNetworkFile, &InInnerPlatformFile, InServerEngineDir, InServerGameDir, InNetworkDoneEvent, InWritingDoneEvent))->StartSynchronousTask();
 }
 
 /**
@@ -883,6 +923,20 @@ void FNetworkPlatformFile::EnsureFileIsLocal(const FString& Filename)
 		}
 	}
 
+	if ( FinishedAsyncNetworkReadUnsolicitedFiles )
+	{
+		delete FinishedAsyncNetworkReadUnsolicitedFiles; // wait here for any async unsolicited files to finish reading being read from the network 
+		FinishedAsyncNetworkReadUnsolicitedFiles = NULL;
+	}
+	if( FinishedAsyncWriteUnsolicitedFiles)
+	{
+		delete FinishedAsyncWriteUnsolicitedFiles; // wait here for any async unsolicited files to finish writing to disk
+		FinishedAsyncWriteUnsolicitedFiles = NULL;
+	}
+	
+
+
+
 	FScopeLock ScopeLock(&SynchronizationObject);
 
 	ThisTime = 1000.0f * float(FPlatformTime::Seconds() - StartTime);
@@ -971,9 +1025,13 @@ void FNetworkPlatformFile::EnsureFileIsLocal(const FString& Filename)
 
 	if (NumUnsolictedFiles)
 	{
-		ReadUnsolicitedFile(NumUnsolictedFiles, *this, *InnerPlatformFile, ServerEngineDir, ServerGameDir);
+		check( FinishedAsyncNetworkReadUnsolicitedFiles == NULL );
+		check( FinishedAsyncWriteUnsolicitedFiles == NULL );
+		FinishedAsyncNetworkReadUnsolicitedFiles = new FScopedEvent;
+		FinishedAsyncWriteUnsolicitedFiles = new FScopedEvent;
+		AsyncReadUnsolicitedFiles(NumUnsolictedFiles, *this, *InnerPlatformFile, ServerEngineDir, ServerGameDir, FinishedAsyncNetworkReadUnsolicitedFiles, FinishedAsyncWriteUnsolicitedFiles);
 	}
-
+	
 	ThisTime = 1000.0f * float(FPlatformTime::Seconds() - StartTime);
 	//UE_LOG(LogNetworkPlatformFile, Display, TEXT("Write file to local %6.2fms"), ThisTime);
 }
diff --git a/Engine/Source/Runtime/NetworkFile/Private/TCPTransport.cpp b/Engine/Source/Runtime/NetworkFile/Private/TCPTransport.cpp
--- a/Engine/Source/Runtime/NetworkFile/Private/TCPTransport.cpp
+++ b/Engine/Source/Runtime/NetworkFile/Private/TCPTransport.cpp
@@ -81,6 +81,26 @@ bool FTCPTransport::SendPayloadAndReceiveResponse(TArray<uint8>& In, TArray<uint
 	return false; 
 }
 
+
+bool FTCPTransport::ReceiveResponse( TArray<uint8> &Out )
+{
+	FArrayReader Response;
+	bool RetResult = true;
+#if USE_MCSOCKET_FOR_NFS
+	RetResult &= FNFSMessageHeader::ReceivePayload(Response, FSimpleAbstractSocket_FMultichannelTCPSocket(MCSocket, NFS_Channels::Main));
+#else
+	RetResult &= FNFSMessageHeader::ReceivePayload(Response, FSimpleAbstractSocket_FSocket(FileSocket));
+#endif
+
+	if (RetResult)
+	{
+		Out.Append( Response.GetData(), Response.Num()); 
+		return true;
+	}
+
+	return false;
+}
+
 FTCPTransport::~FTCPTransport()
 {
 	// on failure, shut it all down
diff --git a/Engine/Source/Runtime/NetworkFile/Private/TCPTransport.h b/Engine/Source/Runtime/NetworkFile/Private/TCPTransport.h
--- a/Engine/Source/Runtime/NetworkFile/Private/TCPTransport.h
+++ b/Engine/Source/Runtime/NetworkFile/Private/TCPTransport.h
@@ -12,8 +12,9 @@ public:
 	~FTCPTransport();
 
 	// ITransport Interface. 
-	virtual bool Initialize(const TCHAR* HostIp);
-	virtual bool SendPayloadAndReceiveResponse(TArray<uint8>& In, TArray<uint8>& Out); 
+	virtual bool Initialize(const TCHAR* HostIp) override;
+	virtual bool SendPayloadAndReceiveResponse(TArray<uint8>& In, TArray<uint8>& Out) override; 
+	virtual bool ReceiveResponse(TArray<uint8> &Out) override;
 
 private: 
 
diff --git a/Engine/Source/Runtime/NetworkFile/Public/NetworkPlatformFile.h b/Engine/Source/Runtime/NetworkFile/Public/NetworkPlatformFile.h
--- a/Engine/Source/Runtime/NetworkFile/Public/NetworkPlatformFile.h
+++ b/Engine/Source/Runtime/NetworkFile/Public/NetworkPlatformFile.h
@@ -110,6 +110,10 @@ public:
 
 
 
+	
+	virtual bool SendPayloadAndReceiveResponse(TArray<uint8>& In, TArray<uint8>& Out);
+	virtual bool ReceiveResponse(TArray<uint8>& Out);
+
 	bool SendReadMessage(uint8* Destination, int64 BytesToRead);
 	bool SendWriteMessage(const uint8* Source, int64 BytesToWrite);
 
@@ -137,12 +141,6 @@ protected:
 
 	virtual void ProcessServerInitialResponse(FArrayReader& InResponse, int32 OutServerPackageVersion, int32 OutServerPackageLicenseeVersion);
 
-		/**
-	 *  Handle HttpRequests This function collates data  puts it in a binary archive. 
-	 * 
-	 */
-	 virtual bool SendPayloadAndReceiveResponse(TArray<uint8>& In, TArray<uint8>& Out);
-
 private:
 
 	/**
@@ -198,8 +196,13 @@ protected:
 	FCriticalSection	LocalDirectoriesCriticalSection;
 	bool				bIsUsable;
 	int32				FileServerPort;
+
 private:
-	
+
+	/* Unsolicitied files events */
+	FScopedEvent *FinishedAsyncNetworkReadUnsolicitedFiles;
+	FScopedEvent *FinishedAsyncWriteUnsolicitedFiles;
+
     // Our network Transport. 
 	class ITransport* Transport; 
 };
diff --git a/Engine/Source/Runtime/NetworkFileSystem/Private/NetworkFileServer.cpp b/Engine/Source/Runtime/NetworkFileSystem/Private/NetworkFileServer.cpp
--- a/Engine/Source/Runtime/NetworkFileSystem/Private/NetworkFileServer.cpp
+++ b/Engine/Source/Runtime/NetworkFileSystem/Private/NetworkFileServer.cpp
@@ -42,21 +42,27 @@ public:
 			  if (!FNFSMessageHeader::ReceivePayload(Payload, FSimpleAbstractSocket_FSocket(Socket)))
 				  break; 
 			  // now process the contents of the payload
-			  FBufferArchive Out;
-			  if ( !FNetworkFileServerClientConnection::ProcessPayload(Payload,Out) )
+			  if ( !FNetworkFileServerClientConnection::ProcessPayload(Payload) )
 			  {
 				  // give the processing of the payload a chance to terminate the connection
 				  // failed to process message
 				  UE_LOG(LogFileServer, Warning, TEXT("Unable to process payload terminating connection"));
 				  break;
 			  }
-			  if ( !FNFSMessageHeader::WrapAndSendPayload(Out, FSimpleAbstractSocket_FSocket(Socket)))
-				  break; 
 		  }
 
 		  return true;
 	}
 
+	virtual bool SendPayload( TArray<uint8> &Out ) override
+	{
+#if USE_MCSOCKET_FOR_NFS
+		return FNFSMessageHeader::WrapAndSendPayload(Out, FSimpleAbstractSocket_FMultichannelTCPSocket(MCSocket, NFS_Channels::Main));
+#else
+		return FNFSMessageHeader::WrapAndSendPayload(Out, FSimpleAbstractSocket_FSocket(Socket));
+#endif
+	}
+
 	virtual void Stop() 
 	{
 		StopRequested.Set(true);
diff --git a/Engine/Source/Runtime/NetworkFileSystem/Private/NetworkFileServerConnection.cpp b/Engine/Source/Runtime/NetworkFileSystem/Private/NetworkFileServerConnection.cpp
--- a/Engine/Source/Runtime/NetworkFileSystem/Private/NetworkFileServerConnection.cpp
+++ b/Engine/Source/Runtime/NetworkFileSystem/Private/NetworkFileServerConnection.cpp
@@ -115,8 +115,9 @@ void FNetworkFileServerClientConnection::ConvertServerFilenameToClientFilename(F
 
 static FCriticalSection SocketCriticalSection;
 
-bool FNetworkFileServerClientConnection::ProcessPayload(FArchive& Ar,FArchive& Out)
+bool FNetworkFileServerClientConnection::ProcessPayload(FArchive& Ar)
 {
+	FBufferArchive Out;
 	bool Result = true;
 
 	// first part of the payload is always the command
@@ -234,17 +235,35 @@ bool FNetworkFileServerClientConnection::ProcessPayload(FArchive& Ar,FArchive& O
 		}
 	}
 
+
 	// send back a reply if the command wrote anything back out
-	if (Out.TotalSize() && bSendUnsolicitedFiles)
+	if (Out.Num() && Result )
 	{
 		int32 NumUnsolictedFiles = UnsolictedFiles.Num();
-		Out << NumUnsolictedFiles;
+		if (bSendUnsolicitedFiles)
+		{
+			Out << NumUnsolictedFiles;
+		}
 
-		for (int32 Index = 0; Index < NumUnsolictedFiles; Index++)
+		UE_LOG(LogFileServer, Verbose, TEXT("Returning payload with %d bytes"), Out.Num());
+
+		// send back a reply
+		Result &= SendPayload( Out );
+
+		if (bSendUnsolicitedFiles && Result )
 		{
-			PackageFile(UnsolictedFiles[Index], Out);
+			for (int32 Index = 0; Index < NumUnsolictedFiles; Index++)
+			{
+				FBufferArchive OutUnsolicitedFile;
+				PackageFile(UnsolictedFiles[Index], OutUnsolicitedFile);
+
+				UE_LOG(LogFileServer, Display, TEXT("Returning unsolicited file %s with %d bytes"), *UnsolictedFiles[Index], OutUnsolicitedFile.Num());
+
+				Result &= SendPayload(OutUnsolicitedFile);
+			}
+
+			UnsolictedFiles.Empty();
 		}
-		UnsolictedFiles.Empty();		
 	}
 
 	UE_LOG(LogFileServer, Verbose, TEXT("Done Processing payload with Cmd %d Total Size sending %d "), Cmd,Out.TotalSize());
@@ -268,6 +287,10 @@ void FNetworkFileServerClientConnection::ProcessOpenFile( FArchive& In, FArchive
 		In << bAllowRead;
 	}
 
+	// todo: clients from the same ip address "could" be trying to write to the same file in the same sandbox (for example multiple windows clients)
+	//			should probably have the sandbox write to separate files for each client
+	//			not important for now
+
 	ConvertClientFilenameToServerFilename(Filename);
 
 	if (bIsWriting)
diff --git a/Engine/Source/Runtime/NetworkFileSystem/Private/NetworkFileServerConnection.h b/Engine/Source/Runtime/NetworkFileSystem/Private/NetworkFileServerConnection.h
--- a/Engine/Source/Runtime/NetworkFileSystem/Private/NetworkFileServerConnection.h
+++ b/Engine/Source/Runtime/NetworkFileSystem/Private/NetworkFileServerConnection.h
@@ -32,7 +32,7 @@ public:
 	 * @param Ar An archive containing the payload data.
 	 * @param Out An archive that will contain the processed data.
 	 */
-	bool ProcessPayload(FArchive& Ar,FArchive& Out);
+	bool ProcessPayload(FArchive& Ar);
 
 	/**
 	 * Gets the client connection's description.
@@ -147,6 +147,9 @@ protected:
 	 */
 	void ProcessSyncFile( FArchive& In, FArchive& Out );
 
+
+	virtual bool SendPayload( TArray<uint8> &Out ) { return false; }
+	
 private:
 
 	// Hold the name of the currently connected platform.
diff --git a/Engine/Source/Runtime/NetworkFileSystem/Private/NetworkFileServerHttp.cpp b/Engine/Source/Runtime/NetworkFileSystem/Private/NetworkFileServerHttp.cpp
--- a/Engine/Source/Runtime/NetworkFileSystem/Private/NetworkFileServerHttp.cpp
+++ b/Engine/Source/Runtime/NetworkFileSystem/Private/NetworkFileServerHttp.cpp
@@ -238,7 +238,7 @@ void FNetworkFileServerHttp::Process(FArchive& In, FArchive&Out, FNetworkFileSer
 			Server->RequestHandlers.Add(ClientGuid,Connection);
 		}
 
-		Connection->ProcessPayload(In, Out);
+		// Connection->ProcessPayload(In, Out);
 	}
 }
 
