Fix cookonthefly TLS crash.

diff --git a/Engine/Source/Runtime/NetworkFile/Private/NetworkPlatformFile.cpp b/Engine/Source/Runtime/NetworkFile/Private/NetworkPlatformFile.cpp
--- a/Engine/Source/Runtime/NetworkFile/Private/NetworkPlatformFile.cpp
+++ b/Engine/Source/Runtime/NetworkFile/Private/NetworkPlatformFile.cpp
@@ -825,7 +825,8 @@ public:
 
 	FORCEINLINE TStatId GetStatId() const
 	{
-		RETURN_QUICK_DECLARE_CYCLE_STAT(FAsyncNetworkWriteWorker, STATGROUP_ThreadPoolAsyncTasks);
+		return TStatId();
+		//RETURN_QUICK_DECLARE_CYCLE_STAT(FAsyncNetworkWriteWorker, STATGROUP_ThreadPoolAsyncTasks);
 	}
 };
 
