Cookon the fly client searches plugin directory for out of date files.

diff --git a/Engine/Source/Runtime/NetworkFile/Private/NetworkPlatformFile.cpp b/Engine/Source/Runtime/NetworkFile/Private/NetworkPlatformFile.cpp
--- a/Engine/Source/Runtime/NetworkFile/Private/NetworkPlatformFile.cpp
+++ b/Engine/Source/Runtime/NetworkFile/Private/NetworkPlatformFile.cpp
@@ -625,10 +625,15 @@ void FNetworkPlatformFile::FillGetFileList(FNetworkFileArchive& Payload, bool bI
 	}
 
 	FString EngineRelPath = FPaths::EngineDir();
+	FString EngineRelPluginPath = FPaths::EnginePluginsDir();
 	FString GameRelPath = FPaths::GameDir();
+	FString GameRelPluginPath = FPaths::GamePluginsDir();
+
 	TArray<FString> Directories;
 	Directories.Add(EngineRelPath);
+	Directories.Add(EngineRelPluginPath);
 	Directories.Add(GameRelPath);
+	Directories.Add(GameRelPluginPath);
 
 	Payload << TargetPlatformNames;
 	Payload << GameName;
