Fixed issue with cook on the fly client  not finding the plugins directory when starting up, which means they will not be cleaned.

diff --git a/Engine/Source/Runtime/NetworkFile/Private/NetworkPlatformFile.cpp b/Engine/Source/Runtime/NetworkFile/Private/NetworkPlatformFile.cpp
--- a/Engine/Source/Runtime/NetworkFile/Private/NetworkPlatformFile.cpp
+++ b/Engine/Source/Runtime/NetworkFile/Private/NetworkPlatformFile.cpp
@@ -172,6 +172,11 @@ void FNetworkPlatformFile::InitializeAfterSetActive()
 			int32 ServerPackageLicenseeVersion = 0;
 			ProcessServerInitialResponse(Response, ServerPackageVersion, ServerPackageLicenseeVersion);
 
+
+			/* The server root content directories */
+			TArray<FString> ServerRootContentDirectories;
+			Response << ServerRootContentDirectories;
+
 			// receive a list of the cache files and their timestamps
 			TMap<FString, FDateTime> ServerCachedFiles;
 			Response << ServerCachedFiles;
@@ -234,13 +239,13 @@ void FNetworkPlatformFile::InitializeAfterSetActive()
 			// use the timestamp grabbing visitor to get all the content times
 			FLocalTimestampDirectoryVisitor Visitor(*InnerPlatformFile, DirectoriesToSkip, DirectoriesToNotRecurse, false);
 
-			TArray<FString> RootContentPaths;
-			FPackageName::QueryRootContentPaths( RootContentPaths );
-			for( TArray<FString>::TConstIterator RootPathIt( RootContentPaths ); RootPathIt; ++RootPathIt )
+			/*TArray<FString> RootContentPaths;
+			FPackageName::QueryRootContentPaths(RootContentPaths); */
+			for (TArray<FString>::TConstIterator RootPathIt(ServerRootContentDirectories); RootPathIt; ++RootPathIt)
 			{
-				const FString& RootPath = *RootPathIt;
-				const FString& ContentFolder = FPackageName::LongPackageNameToFilename(RootPath);
-
+				/*const FString& RootPath = *RootPathIt;
+				const FString& ContentFolder = FPackageName::LongPackageNameToFilename(RootPath);*/
+				const FString& ContentFolder = *RootPathIt;
 				InnerPlatformFile->IterateDirectory( *ContentFolder, Visitor);
 			}
 
diff --git a/Engine/Source/Runtime/NetworkFileSystem/Private/NetworkFileServerConnection.cpp b/Engine/Source/Runtime/NetworkFileSystem/Private/NetworkFileServerConnection.cpp
--- a/Engine/Source/Runtime/NetworkFileSystem/Private/NetworkFileServerConnection.cpp
+++ b/Engine/Source/Runtime/NetworkFileSystem/Private/NetworkFileServerConnection.cpp
@@ -823,6 +823,16 @@ bool FNetworkFileServerClientConnection::ProcessGetFileList( FArchive& In, FArch
 	// Do it again, preventing access to non-cooked files
 	if( bIsStreamingRequest == false )
 	{
+		TArray<FString> RootContentPaths;
+		FPackageName::QueryRootContentPaths(RootContentPaths); 
+		TArray<FString> ContentFolders;
+		for (const auto& RootPath : RootContentPaths)
+		{
+			const FString& ContentFolder = FPackageName::LongPackageNameToFilename(RootPath);
+			ContentFolders.Add(ContentFolder);
+		}
+		Out << ContentFolders;
+
 		// Do it again, preventing access to non-cooked files
 		const int32 NUM_EXCLUSION_WILDCARDS = 2;
 		FString ExclusionWildcard[NUM_EXCLUSION_WILDCARDS];
