Fix cookonthefly plugin discovery

diff --git a/Engine/Source/Runtime/NetworkFile/Private/NetworkPlatformFile.cpp b/Engine/Source/Runtime/NetworkFile/Private/NetworkPlatformFile.cpp
--- a/Engine/Source/Runtime/NetworkFile/Private/NetworkPlatformFile.cpp
+++ b/Engine/Source/Runtime/NetworkFile/Private/NetworkPlatformFile.cpp
@@ -374,7 +374,7 @@ bool FNetworkPlatformFile::MoveFile(const TCHAR* To, const TCHAR* From)
 //	return FNFSMessageHeader::WrapAndSendPayload(Payload, FileSocket);
 
 	FString RelativeFrom = From;
-	FPaths::MakeStandardFilename(RelativeFrom);
+	MakeStandardNetworkFilename(RelativeFrom);
 
 	// don't copy files in local directories
 	if (!IsInLocalDirectory(RelativeFrom))
@@ -406,7 +406,7 @@ IFileHandle* FNetworkPlatformFile::OpenRead(const TCHAR* Filename, bool bAllowWr
 //	FScopeLock ScopeLock(&SynchronizationObject);
 
 	FString RelativeFilename = Filename;
-	FPaths::MakeStandardFilename(RelativeFilename);
+	MakeStandardNetworkFilename(RelativeFilename);
 	// don't copy files in local directories
 	if (!IsInLocalDirectory(RelativeFilename))
 	{
@@ -466,7 +466,7 @@ bool FNetworkPlatformFile::IterateDirectory(const TCHAR* InDirectory, IPlatformF
 
 	// local files go right to the source
 	FString RelativeDirectory = InDirectory;
-	FPaths::MakeStandardFilename(RelativeDirectory);
+	MakeStandardNetworkFilename(RelativeDirectory);
 	if (IsInLocalDirectory(RelativeDirectory))
 	{
 		return InnerPlatformFile->IterateDirectory(InDirectory, Visitor);
@@ -501,7 +501,7 @@ bool FNetworkPlatformFile::IterateDirectoryRecursively(const TCHAR* InDirectory,
 
 	// local files go right to the source
 	FString RelativeDirectory = InDirectory;
-	FPaths::MakeStandardFilename(RelativeDirectory);
+	MakeStandardNetworkFilename(RelativeDirectory);
 
 	if (IsInLocalDirectory(RelativeDirectory))
 	{
@@ -545,7 +545,7 @@ bool FNetworkPlatformFile::CopyFile(const TCHAR* To, const TCHAR* From)
 //	FScopeLock ScopeLock(&SynchronizationObject);
 
 	FString RelativeFrom = From;
-	FPaths::MakeStandardFilename(RelativeFrom);
+	MakeStandardNetworkFilename(RelativeFrom);
 
 	// don't copy files in local directories
 	if (!IsInLocalDirectory(RelativeFrom))
@@ -561,7 +561,7 @@ bool FNetworkPlatformFile::CopyFile(const TCHAR* To, const TCHAR* From)
 FString FNetworkPlatformFile::ConvertToAbsolutePathForExternalAppForRead( const TCHAR* Filename )
 {
 	FString RelativeFrom = Filename;
-	FPaths::MakeStandardFilename(RelativeFrom);
+	MakeStandardNetworkFilename(RelativeFrom);
 	
 	if (!IsInLocalDirectory(RelativeFrom))
 	{
@@ -573,7 +573,7 @@ FString FNetworkPlatformFile::ConvertToAbsolutePathForExternalAppForRead( const
 FString FNetworkPlatformFile::ConvertToAbsolutePathForExternalAppForWrite( const TCHAR* Filename )
 {
 	FString RelativeFrom = Filename;
-	FPaths::MakeStandardFilename(RelativeFrom);
+	MakeStandardNetworkFilename(RelativeFrom);
 
 	if (!IsInLocalDirectory(RelativeFrom))
 	{
@@ -590,7 +590,8 @@ bool FNetworkPlatformFile::DirectoryExists(const TCHAR* Directory)
 	}
 	// If there are any syncable files in this directory, consider it existing
 	FString RelativeDirectory = Directory;
-	FPaths::MakeStandardFilename(RelativeDirectory);
+	MakeStandardNetworkFilename(RelativeDirectory);		
+
 	FServerTOC::FDirectory* ServerDirectory = ServerFiles.FindDirectory(RelativeDirectory);
 	return ServerDirectory != NULL;
 }
@@ -598,7 +599,7 @@ bool FNetworkPlatformFile::DirectoryExists(const TCHAR* Directory)
 void FNetworkPlatformFile::GetFileInfo(const TCHAR* Filename, FFileInfo& Info)
 {
 	FString RelativeFilename = Filename;
-	FPaths::MakeStandardFilename(RelativeFilename);
+	MakeStandardNetworkFilename(RelativeFilename);
 
 	// don't copy files in local directories
 	if (!IsInLocalDirectory(RelativeFilename))
@@ -1048,6 +1049,13 @@ void FNetworkPlatformFile::EnsureFileIsLocal(const FString& Filename)
 	//UE_LOG(LogNetworkPlatformFile, Display, TEXT("Write file to local %6.2fms"), ThisTime);
 }
 
+static FString NetworkPlatformFileEndChop(TEXT("/"));
+void FNetworkPlatformFile::MakeStandardNetworkFilename(FString& Filename)
+{
+	FPaths::MakeStandardFilename(Filename);
+	Filename.RemoveFromEnd(NetworkPlatformFileEndChop, ESearchCase::CaseSensitive);
+}
+
 bool FNetworkPlatformFile::IsInLocalDirectoryUnGuarded(const FString& Filename)
 {
 	// cache the directory of the input file
diff --git a/Engine/Source/Runtime/NetworkFile/Public/NetworkPlatformFile.h b/Engine/Source/Runtime/NetworkFile/Public/NetworkPlatformFile.h
--- a/Engine/Source/Runtime/NetworkFile/Public/NetworkPlatformFile.h
+++ b/Engine/Source/Runtime/NetworkFile/Public/NetworkPlatformFile.h
@@ -163,6 +163,12 @@ private:
 	 */
 	void EnsureFileIsLocal(const FString& Filename);
 
+	/**
+	* Does normal path standardization, and also any extra modifications to make string comparisons against
+	* the internal directory list work properly.
+	*/
+	void MakeStandardNetworkFilename(FString& Filename);
+
 	/**
 	 * This function will send a payload data (with header) and wait for a response, serializing
 	 * the response to a FBufferArchive
