Copying //UE4/Orion-Staging to //UE4/Main (Source: //Orion/Dev-General @ 3483207)

diff --git a/Engine/Source/Runtime/Core/Public/HAL/MallocBinned2.h b/Engine/Source/Runtime/Core/Public/HAL/MallocBinned2.h
--- a/Engine/Source/Runtime/Core/Public/HAL/MallocBinned2.h
+++ b/Engine/Source/Runtime/Core/Public/HAL/MallocBinned2.h
@@ -50,6 +50,21 @@
 #endif
 
 
+#define BINNED2_ALLOCATOR_STATS 1
+
+
+#if BINNED2_ALLOCATOR_STATS
+extern int64 AllocatedSmallPoolMemory; // memory that's requested to be allocated by the game
+
+
+//////////////////////////////////////////////////////////////////////////
+// the following don't need a critical section because they are covered by the critical section called Mutex
+extern int64 AllocatedOSSmallPoolMemory;
+extern int64 AllocatedLargePoolMemory; // memory requests to the OS which don't fit in the small pool
+extern int64 AllocatedLargePoolMemoryWAlignment; // when we allocate at OS level we need to align to a size
+#endif
+
+
 
 //
 // Optimized virtual memory allocator.
@@ -495,6 +510,10 @@ public:
 	void FreeExternal(void *Ptr);
 	bool GetAllocationSizeExternal(void* Ptr, SIZE_T& SizeOut);
 
+	virtual void GetAllocatorStats( FGenericMemoryStats& out_Stats ) override;
+	/** Dumps current allocator stats to the log. */
+	virtual void DumpAllocatorStats(class FOutputDevice& Ar) override;
+	
 	static uint16 SmallBlockSizesReversed[BINNED2_SMALL_POOL_COUNT]; // this is reversed to get the smallest elements on our main cache line
 	static FMallocBinned2* MallocBinned2;
 	static uint32 Binned2TlsSlot;
diff --git a/Engine/Source/Runtime/Core/Public/Misc/EngineVersion.h b/Engine/Source/Runtime/Core/Public/Misc/EngineVersion.h
--- a/Engine/Source/Runtime/Core/Public/Misc/EngineVersion.h
+++ b/Engine/Source/Runtime/Core/Public/Misc/EngineVersion.h
@@ -50,6 +50,9 @@ public:
 		return Branch.Replace( TEXT( "+" ), TEXT( "/" ) );
 	}
 
+	const FString& GetBranchDescriptor() const;
+		
+
 private:
 
 	/** Branch name. */
diff --git a/Engine/Source/Runtime/Core/Public/Misc/TimeGuard.h b/Engine/Source/Runtime/Core/Public/Misc/TimeGuard.h
--- a/Engine/Source/Runtime/Core/Public/Misc/TimeGuard.h
+++ b/Engine/Source/Runtime/Core/Public/Misc/TimeGuard.h
@@ -78,7 +78,6 @@ protected:
 	static CORE_API void						ReportHitch(const TCHAR* InName, float TimeMS);
 	static TMap<const TCHAR*, FGuardInfo>		HitchData;
 	static FCriticalSection						ReportMutex;
-	static CORE_API double						LastHitchTime;
 	static CORE_API bool						bEnabled;
 	static CORE_API float						FrameTimeThresholdMS;
 
diff --git a/Engine/Source/Runtime/CoreUObject/Private/UObject/UObjectBase.cpp b/Engine/Source/Runtime/CoreUObject/Private/UObject/UObjectBase.cpp
--- a/Engine/Source/Runtime/CoreUObject/Private/UObject/UObjectBase.cpp
+++ b/Engine/Source/Runtime/CoreUObject/Private/UObject/UObjectBase.cpp
@@ -934,6 +934,8 @@ static FAutoConsoleVariableRef CMaxObjectsInGame(
 	ECVF_Default
 	);
 
+
+#include "Commandline.h"
 /**
  * Final phase of UObject initialization. all auto register objects are added to the main data structures.
  */
@@ -951,10 +953,20 @@ void UObjectBaseInit()
 	// FPlatformProperties::RequiresCookedData() will be false. Please note that GIsEditor and FApp::IsGame() are not valid at this point.
 	if( FPlatformProperties::RequiresCookedData() )
 	{
-		GConfig->GetInt( TEXT("/Script/Engine.GarbageCollectionSettings"), TEXT("gc.MaxObjectsNotConsideredByGC"), MaxObjectsNotConsideredByGC, GEngineIni );
+		FString Value;
+		bool bIsCookOnTheFly = FParse::Value(FCommandLine::Get(), TEXT("-filehostip="), Value);
+		if (bIsCookOnTheFly)
+		{
+			extern int32 GCreateGCClusters;
+			GCreateGCClusters = false;
+		}
+		else
+		{
+			GConfig->GetInt(TEXT("/Script/Engine.GarbageCollectionSettings"), TEXT("gc.MaxObjectsNotConsideredByGC"), MaxObjectsNotConsideredByGC, GEngineIni);
 
-		// Not used on PC as in-place creation inside bigger pool interacts with the exit purge and deleting UObject directly.
-		GConfig->GetInt( TEXT("/Script/Engine.GarbageCollectionSettings"), TEXT("gc.SizeOfPermanentObjectPool"), SizeOfPermanentObjectPool, GEngineIni );
+			// Not used on PC as in-place creation inside bigger pool interacts with the exit purge and deleting UObject directly.
+			GConfig->GetInt(TEXT("/Script/Engine.GarbageCollectionSettings"), TEXT("gc.SizeOfPermanentObjectPool"), SizeOfPermanentObjectPool, GEngineIni);
+		}
 
 		// Maximum number of UObjects in cooked game
 		GConfig->GetInt(TEXT("/Script/Engine.GarbageCollectionSettings"), TEXT("gc.MaxObjectsInGame"), MaxUObjects, GEngineIni);
diff --git a/Engine/Source/Runtime/CoreUObject/Private/UObject/UObjectClusters.cpp b/Engine/Source/Runtime/CoreUObject/Private/UObject/UObjectClusters.cpp
--- a/Engine/Source/Runtime/CoreUObject/Private/UObject/UObjectClusters.cpp
+++ b/Engine/Source/Runtime/CoreUObject/Private/UObject/UObjectClusters.cpp
@@ -612,7 +612,7 @@ public:
 /** Looks through objects loaded with a package and creates clusters from them */
 void CreateClustersFromPackage(FLinkerLoad* PackageLinker)
 {	
-	if (FPlatformProperties::RequiresCookedData() && !GIsInitialLoad && GCreateGCClusters && !GUObjectArray.IsOpenForDisregardForGC())
+	if (FPlatformProperties::RequiresCookedData() && !GIsInitialLoad && GCreateGCClusters && !GUObjectArray.IsOpenForDisregardForGC() && GUObjectArray.DisregardForGCEnabled() )
 	{
 		check(PackageLinker);
 
diff --git a/Engine/Source/Runtime/Engine/Classes/Components/PrimitiveComponent.h b/Engine/Source/Runtime/Engine/Classes/Components/PrimitiveComponent.h
--- a/Engine/Source/Runtime/Engine/Classes/Components/PrimitiveComponent.h
+++ b/Engine/Source/Runtime/Engine/Classes/Components/PrimitiveComponent.h
@@ -1149,6 +1149,10 @@ public:
 	UFUNCTION(BlueprintCallable, Category="Rendering")
 	void SetCastShadow(bool NewCastShadow);
 
+	/** Changes the value of bSingleSampleShadowFromStationaryLights. */
+	UFUNCTION(BlueprintCallable, Category="Rendering")
+	void SetSingleSampleShadowFromStationaryLights(bool bNewSingleSampleShadowFromStationaryLights);
+
 	/** Changes the value of TranslucentSortPriority. */
 	UFUNCTION(BlueprintCallable, Category="Rendering")
 	void SetTranslucentSortPriority(int32 NewTranslucentSortPriority);
diff --git a/Engine/Source/Runtime/Engine/Classes/Engine/DemoNetDriver.h b/Engine/Source/Runtime/Engine/Classes/Engine/DemoNetDriver.h
--- a/Engine/Source/Runtime/Engine/Classes/Engine/DemoNetDriver.h
+++ b/Engine/Source/Runtime/Engine/Classes/Engine/DemoNetDriver.h
@@ -344,6 +344,7 @@ private:
 	 * Maximum time allowed each frame to spend on saving a checkpoint. If 0, it will save the checkpoint in a single frame, regardless of how long it takes.
 	 * See also demo.CheckpointSaveMaxMSPerFrameOverride.
 	 */
+	UPROPERTY(Config)
 	float CheckpointSaveMaxMSPerFrame;
 
 	/** A player controller that this driver should consider its viewpoint for actor prioritization purposes. */
diff --git a/Engine/Source/Runtime/Engine/Classes/Engine/Engine.h b/Engine/Source/Runtime/Engine/Classes/Engine/Engine.h
--- a/Engine/Source/Runtime/Engine/Classes/Engine/Engine.h
+++ b/Engine/Source/Runtime/Engine/Classes/Engine/Engine.h
@@ -559,7 +559,9 @@ class IAnalyticsProvider;
 DECLARE_DELEGATE_OneParam(FBeginStreamingPauseDelegate, FViewport*);
 DECLARE_DELEGATE(FEndStreamingPauseDelegate);
 
-DECLARE_MULTICAST_DELEGATE_OneParam(FEngineHitchDetectedDelegate, float /*HitchDurationInSeconds*/);
+enum class EFrameHitchType: uint8;
+
+DECLARE_MULTICAST_DELEGATE_TwoParams(FEngineHitchDetectedDelegate, EFrameHitchType /*HitchType*/, float /*HitchDurationInSeconds*/);
 
 
 /**
@@ -2060,7 +2062,7 @@ public:
 	* Attempts to reclaim any idle memory by performing a garbage collection and broadcasting FCoreDelegates::OnMemoryTrim. Pending rendering commands are first flushed. This is called
 	* between level loads and may be called at other times, but is expensive and should be used sparingly. Do
 	*/
-	void TrimMemory();
+	static void TrimMemory();
 
 	/**
 	 * Calculates information about the previous frame and passes it to all active performance data consumers.
diff --git a/Engine/Source/Runtime/Engine/Classes/GameFramework/CharacterMovementComponent.h b/Engine/Source/Runtime/Engine/Classes/GameFramework/CharacterMovementComponent.h
--- a/Engine/Source/Runtime/Engine/Classes/GameFramework/CharacterMovementComponent.h
+++ b/Engine/Source/Runtime/Engine/Classes/GameFramework/CharacterMovementComponent.h
@@ -2180,6 +2180,11 @@ public:
 	void ClientAdjustRootMotionSourcePosition(float TimeStamp, FRootMotionSourceGroup ServerRootMotion, bool bHasAnimRootMotion, float ServerMontageTrackPosition, FVector ServerLoc, FVector_NetQuantizeNormal ServerRotation, float ServerVelZ, UPrimitiveComponent* ServerBase, FName ServerBoneName, bool bHasBase, bool bBaseRelativePosition, uint8 ServerMovementMode);
 	void ClientAdjustRootMotionSourcePosition_Implementation(float TimeStamp, FRootMotionSourceGroup ServerRootMotion, bool bHasAnimRootMotion, float ServerMontageTrackPosition, FVector ServerLoc, FVector_NetQuantizeNormal ServerRotation, float ServerVelZ, UPrimitiveComponent* ServerBase, FName ServerBoneName, bool bHasBase, bool bBaseRelativePosition, uint8 ServerMovementMode);
 
+protected:
+
+	/** Event notification when client receives a correction from the server. Base implementation logs relevant data and draws debug info if "p.NetShowCorrections" is not equal to 0. */
+	virtual void OnClientCorrectionReceived(class FNetworkPredictionData_Client_Character& ClientData, float TimeStamp, FVector NewLocation, FVector NewVelocity, UPrimitiveComponent* NewBase, FName NewBaseBoneName, bool bHasBase, bool bBaseRelativePosition, uint8 ServerMovementMode);
+
 	// Root Motion
 public:
 	/** Root Motion Group containing active root motion sources being applied to movement */
diff --git a/Engine/Source/Runtime/Engine/Private/AI/Navigation/NavigationModifier.cpp b/Engine/Source/Runtime/Engine/Private/AI/Navigation/NavigationModifier.cpp
--- a/Engine/Source/Runtime/Engine/Private/AI/Navigation/NavigationModifier.cpp
+++ b/Engine/Source/Runtime/Engine/Private/AI/Navigation/NavigationModifier.cpp
@@ -796,10 +796,9 @@ void FCompositeNavModifier::CreateAreaModifiers(const UPrimitiveComponent* PrimC
 	for (int32 Idx = 0; Idx < BodySetup->AggGeom.SphylElems.Num(); Idx++)
 	{
 		const FKSphylElem& SphylElem = BodySetup->AggGeom.SphylElems[Idx];
-		const float CapsuleHeight = SphylElem.Length + SphylElem.Radius;
-		const FTransform AreaOffset(FVector(0, 0, -CapsuleHeight));
+		const FTransform AreaOffset(FVector(0, 0, -SphylElem.Length));
 
-		FAreaNavModifier AreaMod(SphylElem.Radius, CapsuleHeight, AreaOffset * PrimComp->GetComponentTransform(), AreaClass);
+		FAreaNavModifier AreaMod(SphylElem.Radius, SphylElem.Length * 2.0f, AreaOffset * PrimComp->GetComponentTransform(), AreaClass);
 		Add(AreaMod);
 	}
 
@@ -814,8 +813,9 @@ void FCompositeNavModifier::CreateAreaModifiers(const UPrimitiveComponent* PrimC
 	for (int32 Idx = 0; Idx < BodySetup->AggGeom.SphereElems.Num(); Idx++)
 	{
 		const FKSphereElem& SphereElem = BodySetup->AggGeom.SphereElems[Idx];
-		
-		FAreaNavModifier AreaMod(SphereElem.Radius, SphereElem.Radius, PrimComp->GetComponentTransform(), AreaClass);
+		const FTransform AreaOffset(FVector(0, 0, -SphereElem.Radius));
+
+		FAreaNavModifier AreaMod(SphereElem.Radius, SphereElem.Radius * 2.0f, AreaOffset * PrimComp->GetComponentTransform(), AreaClass);
 		Add(AreaMod);
 	}
 }
diff --git a/Engine/Source/Runtime/Engine/Private/Animation/AnimCompressionDerivedData.cpp b/Engine/Source/Runtime/Engine/Private/Animation/AnimCompressionDerivedData.cpp
--- a/Engine/Source/Runtime/Engine/Private/Animation/AnimCompressionDerivedData.cpp
+++ b/Engine/Source/Runtime/Engine/Private/Animation/AnimCompressionDerivedData.cpp
@@ -53,9 +53,10 @@ FString FDerivedDataAnimationCompression::GetPluginSpecificCacheKeySuffix() cons
 	char AdditiveType = bCanBakeAdditive ? NibbleToTChar(OriginalAnimSequence->AdditiveAnimType) : '0';
 	char RefType = bCanBakeAdditive ? NibbleToTChar(OriginalAnimSequence->RefPoseType) : '0';
 
-	FString Ret = FString::Printf(TEXT("%i_%i_%s%s%s_%c%c%i_%s_%s"),
+	FString Ret = FString::Printf(TEXT("%i_%i_%i_%s%s%s_%c%c%i_%s_%s"),
 		(int32)UE_ANIMCOMPRESSION_DERIVEDDATA_VER,
 		(int32)CURRENT_ANIMATION_ENCODING_PACKAGE_VERSION,
+		OriginalAnimSequence->CompressCommandletVersion,
 		*OriginalAnimSequence->GetRawDataGuid().ToString(),
 		*OriginalAnimSequence->GetSkeleton()->GetGuid().ToString(),
 		*OriginalAnimSequence->GetSkeleton()->GetVirtualBoneGuid().ToString(),
@@ -67,7 +68,6 @@ FString FDerivedDataAnimationCompression::GetPluginSpecificCacheKeySuffix() cons
 		);
 
 	return Ret;
-
 }
 
 bool FDerivedDataAnimationCompression::Build( TArray<uint8>& OutData )
diff --git a/Engine/Source/Runtime/Engine/Private/Animation/AnimSequence.cpp b/Engine/Source/Runtime/Engine/Private/Animation/AnimSequence.cpp
--- a/Engine/Source/Runtime/Engine/Private/Animation/AnimSequence.cpp
+++ b/Engine/Source/Runtime/Engine/Private/Animation/AnimSequence.cpp
@@ -566,18 +566,18 @@ void UAnimSequence::PostLoad()
 #if WITH_EDITOR
 void ShowResaveMessage(const UAnimSequence* Sequence)
 {
-	if (!IsRunningGame())
+	if (IsRunningCommandlet())
 	{
-		UE_LOG(LogAnimation, Warning, TEXT("Resave Animation Required(%s, %s): Fixing track data and recompressing."), *GetNameSafe(Sequence), *Sequence->GetPathName());
+		UE_LOG(LogAnimation, Log, TEXT("Resave Animation Required(%s, %s): Fixing track data and recompressing."), *GetNameSafe(Sequence), *Sequence->GetPathName());
 
-		static FName NAME_LoadErrors("LoadErrors");
+		/*static FName NAME_LoadErrors("LoadErrors");
 		FMessageLog LoadErrors(NAME_LoadErrors);
 
 		TSharedRef<FTokenizedMessage> Message = LoadErrors.Warning();
 		Message->AddToken(FTextToken::Create(LOCTEXT("AnimationNeedsResave1", "The Animation ")));
 		Message->AddToken(FAssetNameToken::Create(Sequence->GetPathName(), FText::FromString(GetNameSafe(Sequence))));
 		Message->AddToken(FTextToken::Create(LOCTEXT("AnimationNeedsResave2", " needs resave.")));
-		LoadErrors.Notify();
+		LoadErrors.Notify();*/
 	}
 }
 
diff --git a/Engine/Source/Runtime/Engine/Private/Animation/AnimationRuntime.cpp b/Engine/Source/Runtime/Engine/Private/Animation/AnimationRuntime.cpp
--- a/Engine/Source/Runtime/Engine/Private/Animation/AnimationRuntime.cpp
+++ b/Engine/Source/Runtime/Engine/Private/Animation/AnimationRuntime.cpp
@@ -945,7 +945,21 @@ void FAnimationRuntime::BlendMeshPosesPerBoneWeights(
 		/*out*/ FCompactPose& OutPose,
 		/*out*/ struct FBlendedCurve& OutCurve)
 {
-	check(BasePose.GetNumBones() == BoneBlendWeights.Num());
+	const int32 NumBones = BasePose.GetNumBones();
+	check(BoneBlendWeights.Num() == NumBones);
+	check(OutPose.GetNumBones() == NumBones);
+
+	const int32 NumPoses = BlendPoses.Num();
+	for (const FPerBoneBlendWeight& PerBoneBlendWeight : BoneBlendWeights)
+	{
+		check(PerBoneBlendWeight.SourceIndex >= 0);
+		check(PerBoneBlendWeight.SourceIndex < NumPoses);
+	}
+
+	for (const FCompactPose& BlendPose : BlendPoses)
+	{
+		check(BlendPose.GetNumBones() == NumBones);
+	}
 
 	const FBoneContainer& BoneContainer = BasePose.GetBoneContainer();
 
@@ -953,16 +967,14 @@ void FAnimationRuntime::BlendMeshPosesPerBoneWeights(
 	TCustomBoneIndexArray<FQuat, FCompactPoseBoneIndex> BlendRotations;
 	TCustomBoneIndexArray<FQuat, FCompactPoseBoneIndex> TargetRotations;
 
-	SourceRotations.AddUninitialized(BasePose.GetNumBones());
-	BlendRotations.AddUninitialized(BasePose.GetNumBones());
-	TargetRotations.AddUninitialized(BasePose.GetNumBones());
-
-	int32 PoseNum = BlendPoses.Num();
+	SourceRotations.AddUninitialized(NumBones);
+	BlendRotations.AddUninitialized(NumBones);
+	TargetRotations.AddUninitialized(NumBones);
 
 	TArray<float> MaxPoseWeights;
-	MaxPoseWeights.AddZeroed(PoseNum);
+	MaxPoseWeights.AddZeroed(NumPoses);
 
-	for (FCompactPoseBoneIndex BoneIndex : BasePose.ForEachBoneIndex())
+	for (const FCompactPoseBoneIndex BoneIndex : BasePose.ForEachBoneIndex())
 	{
 		const int32 PoseIndex = BoneBlendWeights[BoneIndex.GetInt()].SourceIndex;
 		const FCompactPoseBoneIndex ParentIndex = BoneContainer.GetParentBoneIndex(BoneIndex);
@@ -1033,13 +1045,13 @@ void FAnimationRuntime::BlendMeshPosesPerBoneWeights(
 		TArray<const FBlendedCurve*> SourceCurves;
 		TArray<float> SourceWegihts;
 
-		SourceCurves.SetNumUninitialized(PoseNum+1);
-		SourceWegihts.SetNumUninitialized(PoseNum+1);
+		SourceCurves.SetNumUninitialized(NumPoses+1);
+		SourceWegihts.SetNumUninitialized(NumPoses +1);
 
 		SourceCurves[0] = &BaseCurve;
 		SourceWegihts[0] = 1.f;
 
-		for(int32 Idx=0; Idx<PoseNum; ++Idx)
+		for(int32 Idx=0; Idx<NumPoses; ++Idx)
 		{
 			SourceCurves[Idx+1] = &BlendedCurves[Idx];
 			SourceWegihts[Idx+1] = MaxPoseWeights[Idx];
@@ -1059,11 +1071,24 @@ void FAnimationRuntime::BlendLocalPosesPerBoneWeights(
 	/*out*/ FCompactPose& OutPose, 
 	/*out*/ struct FBlendedCurve& OutCurve)
 {
-	check(BasePose.GetNumBones() == BoneBlendWeights.Num());
-	int32 PoseNum = BlendPoses.Num();
+	const int32 NumBones = BasePose.GetNumBones();
+	check(BoneBlendWeights.Num() == NumBones);
+	check(OutPose.GetNumBones() == NumBones);
+
+	const int32 NumPoses = BlendPoses.Num();
+	for (const FPerBoneBlendWeight& PerBoneBlendWeight : BoneBlendWeights)
+	{
+		check(PerBoneBlendWeight.SourceIndex >= 0);
+		check(PerBoneBlendWeight.SourceIndex < NumPoses);
+	}
+
+	for (const FCompactPose& BlendPose : BlendPoses)
+	{
+		check(BlendPose.GetNumBones() == NumBones);
+	}
 
 	TArray<float> MaxPoseWeights;
-	MaxPoseWeights.AddZeroed(PoseNum);
+	MaxPoseWeights.AddZeroed(NumPoses);
 
 	for (FCompactPoseBoneIndex BoneIndex : BasePose.ForEachBoneIndex())
 	{
@@ -1097,13 +1122,13 @@ void FAnimationRuntime::BlendLocalPosesPerBoneWeights(
 		TArray<const FBlendedCurve*> SourceCurves;
 		TArray<float> SourceWegihts;
 
-		SourceCurves.SetNumUninitialized(PoseNum+1);
-		SourceWegihts.SetNumUninitialized(PoseNum+1);
+		SourceCurves.SetNumUninitialized(NumPoses +1);
+		SourceWegihts.SetNumUninitialized(NumPoses +1);
 
 		SourceCurves[0] = &BaseCurve;
 		SourceWegihts[0] = 1.f;
 
-		for (int32 Idx=0; Idx<PoseNum; ++Idx)
+		for (int32 Idx=0; Idx<NumPoses; ++Idx)
 		{
 			SourceCurves[Idx+1] = &BlendedCurves[Idx];
 			SourceWegihts[Idx+1] = MaxPoseWeights[Idx];
@@ -1118,14 +1143,16 @@ void FAnimationRuntime::UpdateDesiredBoneWeight(const TArray<FPerBoneBlendWeight
 	// in the future, cache this outside
 	ensure (TargetBoneBlendWeights.Num() == SrcBoneBlendWeights.Num());
 
-	FMemory::Memset(TargetBoneBlendWeights.GetData(), 0, TargetBoneBlendWeights.Num() * sizeof(FPerBoneBlendWeight));
+	FMemory::Memzero(TargetBoneBlendWeights.GetData(), TargetBoneBlendWeights.Num() * sizeof(FPerBoneBlendWeight));
 
 	for (int32 BoneIndex = 0; BoneIndex < SrcBoneBlendWeights.Num(); ++BoneIndex)
 	{
-		int32 PoseIndex = SrcBoneBlendWeights[BoneIndex].SourceIndex;
-		float TargetBlendWeight = BlendWeights[PoseIndex]*SrcBoneBlendWeights[BoneIndex].BlendWeight;
+		const int32 PoseIndex = SrcBoneBlendWeights[BoneIndex].SourceIndex;
+		check(PoseIndex < BlendWeights.Num());
+		float TargetBlendWeight = BlendWeights[PoseIndex] * SrcBoneBlendWeights[BoneIndex].BlendWeight;
+		
 		// if relevant, otherwise all initialized as zero
-		if (TargetBlendWeight > ZERO_ANIMWEIGHT_THRESH)
+		if (FAnimWeight::IsRelevant(TargetBlendWeight))
 		{
 			TargetBoneBlendWeights[BoneIndex].SourceIndex = PoseIndex;
 			TargetBoneBlendWeights[BoneIndex].BlendWeight = TargetBlendWeight;
@@ -1171,9 +1198,9 @@ void FAnimationRuntime::CreateMaskWeights(TArray<FPerBoneBlendWeight>& BoneBlend
 	{
 		const FReferenceSkeleton& RefSkeleton = Skeleton->GetReferenceSkeleton();
 		
-		const int32 TotalNum = RefSkeleton.GetNum();
-		BoneBlendWeights.Reset(TotalNum);
-		BoneBlendWeights.AddZeroed(TotalNum);
+		const int32 NumBones = RefSkeleton.GetNum();
+		BoneBlendWeights.Reset(NumBones);
+		BoneBlendWeights.AddZeroed(NumBones);
 
 		// base mask bone
 		for (int32 PoseIndex=0; PoseIndex<BlendFilters.Num(); ++PoseIndex)
@@ -1183,24 +1210,28 @@ void FAnimationRuntime::CreateMaskWeights(TArray<FPerBoneBlendWeight>& BoneBlend
 			for (int32 BranchIndex=0; BranchIndex<BlendPose.BranchFilters.Num(); ++BranchIndex)
 			{
 				const FBranchFilter& BranchFilter = BlendPose.BranchFilters[BranchIndex];
-				int32 MaskBoneIndex = RefSkeleton.FindBoneIndex(BranchFilter.BoneName);
-
-				// how much weight increase Per depth
-				float IncreaseWeightPerDepth = (BranchFilter.BlendDepth != 0) ? (1.f/((float)BranchFilter.BlendDepth)) : 1.f;
+				const int32 MaskBoneIndex = RefSkeleton.FindBoneIndex(BranchFilter.BoneName);
 
-				// go through skeleton tree requiredboneindices
-				for (int32 BoneIndex = 0; BoneIndex < TotalNum; ++BoneIndex)
+				if (MaskBoneIndex != INDEX_NONE)
 				{
-					// if Depth == -1, it's not a child
-					int32 Depth = RefSkeleton.GetDepthBetweenBones(BoneIndex, MaskBoneIndex);
-					if (Depth != -1)
-					{
-						// when you write to buffer, you'll need to match with BasePoses BoneIndex
-						FPerBoneBlendWeight& BoneBlendWeight = BoneBlendWeights[BoneIndex];
+					// how much weight increase Per depth
+					const float IncreaseWeightPerDepth = (BranchFilter.BlendDepth != 0) ? (1.f/((float)BranchFilter.BlendDepth)) : 1.f;
 
-						BoneBlendWeight.SourceIndex = PoseIndex;
-						float BlendIncrease = IncreaseWeightPerDepth * (float)(Depth + 1);
-						BoneBlendWeight.BlendWeight = FMath::Clamp<float>(BoneBlendWeight.BlendWeight + BlendIncrease, 0.f, 1.f);
+					// go through skeleton bone hierarchy.
+					// Bones are ordered, parents before children. So we can start looking at MaskBoneIndex for children.
+					for (int32 BoneIndex = MaskBoneIndex; BoneIndex < NumBones; ++BoneIndex)
+					{
+						// if Depth == -1, it's not a child
+						const int32 Depth = RefSkeleton.GetDepthBetweenBones(BoneIndex, MaskBoneIndex);
+						if (Depth != -1)
+						{
+							// when you write to buffer, you'll need to match with BasePoses BoneIndex
+							FPerBoneBlendWeight& BoneBlendWeight = BoneBlendWeights[BoneIndex];
+
+							BoneBlendWeight.SourceIndex = PoseIndex;
+							const float BlendIncrease = IncreaseWeightPerDepth * (float)(Depth + 1);
+							BoneBlendWeight.BlendWeight = FMath::Clamp<float>(BoneBlendWeight.BlendWeight + BlendIncrease, 0.f, 1.f);
+						}
 					}
 				}
 			}
diff --git a/Engine/Source/Runtime/Engine/Private/Animation/AnimationUtils.cpp b/Engine/Source/Runtime/Engine/Private/Animation/AnimationUtils.cpp
--- a/Engine/Source/Runtime/Engine/Private/Animation/AnimationUtils.cpp
+++ b/Engine/Source/Runtime/Engine/Private/Animation/AnimationUtils.cpp
@@ -441,13 +441,14 @@ bool FAnimationUtils::GetForcedRecompressionSetting()
 	#define TRYCOMPRESSION_INNER(compressionname,winningcompressor_count,winningcompressor_error,winningcompressor_margin,compressionalgorithm)				\
 {																																							\
 	/* try the alternative compressor	*/																													\
-	(compressionalgorithm)->Reduce( AnimSeq, CompressContext );																								\
+	AnimSeq->CompressionScheme = (compressionalgorithm);																									\
+	(compressionalgorithm)->Reduce(AnimSeq, CompressContext);																								\
 	AnimSeq->SetUseRawDataOnly(false);																														\
 	const SIZE_T NewSize = AnimSeq->GetResourceSizeBytes(EResourceSizeMode::Inclusive);																		\
 																																							\
 	/* compute the savings and compression error*/																											\
-	const SIZE_T MemorySavingsFromOriginal = OriginalSize - NewSize;																						\
-	const SIZE_T MemorySavingsFromPrevious = CurrentSize - NewSize;																							\
+	const int32 MemorySavingsFromOriginal = OriginalSize - NewSize;																						\
+	const int32 MemorySavingsFromPrevious = CurrentSize - NewSize;																							\
 	PctSaving = 0.f;																																		\
 	/* figure out our new compression error*/																												\
 	FAnimationUtils::ComputeCompressionError(AnimSeq, BoneData, NewErrorStats);																				\
@@ -457,14 +458,27 @@ bool FAnimationUtils::GetForcedRecompressionSetting()
 																																							\
 	/* keep it if it we want to force the error below the threshold and it reduces error */																	\
 	bKeepNewCompressionMethod = false;																														\
-	bKeepNewCompressionMethod |= (bLowersError && (WinningCompressorError > MasterTolerance) && bForceBelowThreshold);										\
+	const bool bReducesErrorBelowThreshold = (bLowersError && (WinningCompressorError > MasterTolerance) && bForceBelowThreshold);							\
+	bKeepNewCompressionMethod |= bReducesErrorBelowThreshold;																								\
 	/* or if has an acceptable error and saves space  */																									\
-	bKeepNewCompressionMethod |= bErrorUnderThreshold && (MemorySavingsFromPrevious > 0);																	\
+	const bool bHasAcceptableErrorAndSavesSpace = bErrorUnderThreshold && (MemorySavingsFromPrevious > 0);													\
+	bKeepNewCompressionMethod |= bHasAcceptableErrorAndSavesSpace;																							\
 	/* or if saves the same amount and an acceptable error that is lower than the previous best */															\
-	bKeepNewCompressionMethod |= bErrorUnderThreshold && bLowersError && (MemorySavingsFromPrevious >= 0);													\
+	const bool bLowersErrorAndSavesSameOrBetter = bErrorUnderThreshold && bLowersError && (MemorySavingsFromPrevious >= 0);									\
+	bKeepNewCompressionMethod |= bLowersErrorAndSavesSameOrBetter;																							\
+																																							\
+	PctSaving = (OriginalSize > 0) ? (100.f - (100.f * float(NewSize) / float(OriginalSize))) : 0.f;														\
+	UE_LOG(LogAnimation, Verbose, TEXT("- %s - bytes saved(%i) (%.1f%%) from previous(%i) MaxError(%.2f) bLowersError(%d) %s"),								\
+		compressionname, MemorySavingsFromOriginal, PctSaving, MemorySavingsFromPrevious, NewErrorStats.MaxError, bLowersError, bKeepNewCompressionMethod ? TEXT("(**Best so far**)") : TEXT(""));									\
+																																							\
+	UE_LOG(LogAnimation, Verbose, TEXT("    bReducesErrorBelowThreshold(%d) bHasAcceptableErrorAndSavesSpace(%d) bLowersErrorAndSavesSameOrBetter(%d)"),	\
+		bReducesErrorBelowThreshold, bHasAcceptableErrorAndSavesSpace, bLowersErrorAndSavesSameOrBetter);													\
+																																							\
+	UE_LOG(LogAnimation, Verbose, TEXT("    WinningCompressorError(%f) MasterTolerance(%f) bForceBelowThreshold(%d) bErrorUnderThreshold(%d)"),				\
+		WinningCompressorError, MasterTolerance, bForceBelowThreshold, bErrorUnderThreshold);																\
 																																							\
 	if (bKeepNewCompressionMethod)																															\
-			{																																				\
+	{																																						\
 		WinningCompressorMarginalSavings = MemorySavingsFromPrevious;																						\
 		WinningCompressorCounter = &(winningcompressor_count);																								\
 		WinningCompressorErrorSum = &(winningcompressor_error);																								\
@@ -473,13 +487,21 @@ bool FAnimationUtils::GetForcedRecompressionSetting()
 		CurrentSize = NewSize;																																\
 		WinningCompressorSavings = MemorySavingsFromOriginal;																								\
 		WinningCompressorError = NewErrorStats.MaxError;																									\
-			}																																				\
 																																							\
-	PctSaving = OriginalSize > 0 ? 100.f - (100.f * float(NewSize) / float(OriginalSize)) : 0.f;															\
-	UE_LOG(LogAnimation, Warning, TEXT("- %s - bytes saved: %i (%3.1f%% saved), maxdiff: %f %s"),															\
-	compressionname, MemorySavingsFromOriginal, PctSaving, NewErrorStats.MaxError, bKeepNewCompressionMethod ? TEXT("(**Best so far**)") : TEXT(""));		\
-																																							\
-	if( !bKeepNewCompressionMethod )																														\
+		/* backup key information from the sequence */																										\
+		SavedCompressionScheme = AnimSeq->CompressionScheme;																								\
+		SavedTranslationCompressionFormat = AnimSeq->TranslationCompressionFormat;																			\
+		SavedRotationCompressionFormat = AnimSeq->RotationCompressionFormat;																				\
+		SavedKeyEncodingFormat = AnimSeq->KeyEncodingFormat;																								\
+		SavedCompressedTrackOffsets = AnimSeq->CompressedTrackOffsets;																						\
+		SavedCompressedScaleOffsets = AnimSeq->CompressedScaleOffsets;																						\
+		SavedCompressedByteStream = AnimSeq->CompressedByteStream;																							\
+		SavedTranslationCodec = AnimSeq->TranslationCodec;																									\
+		SavedRotationCodec = AnimSeq->RotationCodec;																										\
+		SavedScaleCodec = AnimSeq->ScaleCodec;																												\
+		bSavedUseRawDataOnly = false;																														\
+	}																																						\
+	else																																					\
 	{																																						\
 		/* revert back to the old method by copying back the data we cached */																				\
 		AnimSeq->CompressionScheme = SavedCompressionScheme;																								\
@@ -498,21 +520,6 @@ bool FAnimationUtils::GetForcedRecompressionSetting()
 		const SIZE_T RestoredSize = AnimSeq->GetResourceSizeBytes(EResourceSizeMode::Inclusive);															\
 		check(RestoredSize == CurrentSize);																													\
 	}																																						\
-	else																																					\
-	{																																						\
-		/* backup key information from the sequence */																										\
-		SavedCompressionScheme				= AnimSeq->CompressionScheme;																					\
-		SavedTranslationCompressionFormat	= AnimSeq->TranslationCompressionFormat;																		\
-		SavedRotationCompressionFormat		= AnimSeq->RotationCompressionFormat;																			\
-		SavedKeyEncodingFormat				= AnimSeq->KeyEncodingFormat;																					\
-		SavedCompressedTrackOffsets			= AnimSeq->CompressedTrackOffsets;																				\
-		SavedCompressedScaleOffsets			= AnimSeq->CompressedScaleOffsets;																				\
-		SavedCompressedByteStream			= AnimSeq->CompressedByteStream;																				\
-		SavedTranslationCodec				= AnimSeq->TranslationCodec;																					\
-		SavedRotationCodec					= AnimSeq->RotationCodec;																						\
-		SavedScaleCodec						= AnimSeq->ScaleCodec;																							\
-		bSavedUseRawDataOnly				= false;																										\
-	}																																						\
 }
 #endif
 
@@ -526,9 +533,6 @@ bool FAnimationUtils::GetForcedRecompressionSetting()
 	static float Algorithm ## CompressorSumError = 0.0f; \
 	static int32 Algorithm ## CompressorWinMargin = 0
 
-/** Control animation recompression upon load. */
-bool GDisableAnimationRecompression	= false;
-
 /**
  * Utility function to compress an animation. If the animation is currently associated with a codec, it will be used to 
  * compress the animation. Otherwise, the default codec will be used. If AllowAlternateCompressor is true, an
@@ -616,11 +620,6 @@ void FAnimationUtils::CompressAnimSequenceExplicit(
 	const bool bTryIntervalKeyRemoval)
 {
 #if WITH_EDITORONLY_DATA
-	if( GDisableAnimationRecompression )
-	{
-		return;
-	}
-
 	DECLARE_ANIM_COMP_ALGORITHM(BitwiseACF_Float96);
 	DECLARE_ANIM_COMP_ALGORITHM(BitwiseACF_Fixed48);
 	DECLARE_ANIM_COMP_ALGORITHM(BitwiseACF_IntervalFixed32);
@@ -692,11 +691,11 @@ void FAnimationUtils::CompressAnimSequenceExplicit(
 		bool const bTryAlternateCompressor = MasterTolerance > 0.0f;
 
 		// Get the current size
-		int32 OriginalSize = AnimSeq->GetResourceSizeBytes(EResourceSizeMode::Inclusive);
+		SIZE_T OriginalSize = AnimSeq->GetResourceSizeBytes(EResourceSizeMode::Inclusive);
 		TotalSizeBefore += OriginalSize;
 
 		// Estimate total uncompressed
-		TotalUncompressed += ((sizeof(FVector) + sizeof(FQuat)) * NumRawDataTracks * AnimSeq->NumFrames);
+		TotalUncompressed += ((sizeof(FVector) + sizeof(FQuat) + sizeof(FVector)) * NumRawDataTracks * AnimSeq->NumFrames);
 
 		// Filter RAW data to get rid of mismatched tracks (translation/rotation data with a different number of keys than there are frames)
 		// No trivial key removal is done at this point (impossible error metrics of -1), since all of the techniques will perform it themselves
@@ -721,7 +720,8 @@ void FAnimationUtils::CompressAnimSequenceExplicit(
 		{
 			UAnimCompress* OriginalCompressionAlgorithm = AnimSeq->CompressionScheme ? AnimSeq->CompressionScheme : FAnimationUtils::GetDefaultAnimationCompressionAlgorithm();
 
-			OriginalCompressionAlgorithm->Reduce( AnimSeq, CompressContext );
+			AnimSeq->CompressionScheme = OriginalCompressionAlgorithm;
+			OriginalCompressionAlgorithm->Reduce(AnimSeq, CompressContext);
 			AnimSeq->SetUseRawDataOnly(false);
 			AfterOriginalRecompression = AnimSeq->GetResourceSizeBytes(EResourceSizeMode::Inclusive);
 
@@ -735,7 +735,9 @@ void FAnimationUtils::CompressAnimSequenceExplicit(
 		}
  
 		// Check for global permission to try an alternative compressor
-		if( bTryAlternateCompressor && !AnimSeq->bDoNotOverrideCompression )
+		// we don't check for bDoNotOverrideCompression here, as that is now used as part of the UAnimCompress_Automatic compressor
+		// And it's valid to manually recompress animations
+		if( bTryAlternateCompressor /*&& (!AnimSeq->bDoNotOverrideCompression */)
 		{
 			AnimationErrorStats NewErrorStats = OriginalErrorStats;
 			if (bRaiseMaxErrorToExisting)
@@ -765,7 +767,7 @@ void FAnimationUtils::CompressAnimSequenceExplicit(
 				++TotalRecompressions;
 
 				// Prepare to compress
-				int32 CurrentSize = AnimSeq->GetResourceSizeBytes(EResourceSizeMode::Inclusive);
+				SIZE_T CurrentSize = AnimSeq->GetResourceSizeBytes(EResourceSizeMode::Inclusive);
 				int32* WinningCompressorCounter = NULL;
 				float* WinningCompressorErrorSum = NULL;
 				int32* WinningCompressorMarginalSavingsSum = NULL;
diff --git a/Engine/Source/Runtime/Engine/Private/BoneContainer.cpp b/Engine/Source/Runtime/Engine/Private/BoneContainer.cpp
--- a/Engine/Source/Runtime/Engine/Private/BoneContainer.cpp
+++ b/Engine/Source/Runtime/Engine/Private/BoneContainer.cpp
@@ -166,14 +166,16 @@ void FBoneContainer::Initialize()
 		SkeletonToCompactPose.Add(FCompactPoseBoneIndex(CompactIndex));
 	}
 
-
 	for (const FVirtualBoneRefData& VBRefBone : RefSkeleton->GetVirtualBoneRefData())
 	{
-		int32 VBInd = MeshIndexToCompactPoseIndex[VBRefBone.VBRefSkelIndex];
-		int32 SourceInd = MeshIndexToCompactPoseIndex[VBRefBone.SourceRefSkelIndex];
-		int32 TargetInd = MeshIndexToCompactPoseIndex[VBRefBone.TargetRefSkelIndex];
+		const int32 VBInd = MeshIndexToCompactPoseIndex[VBRefBone.VBRefSkelIndex];
+		const int32 SourceInd = MeshIndexToCompactPoseIndex[VBRefBone.SourceRefSkelIndex];
+		const int32 TargetInd = MeshIndexToCompactPoseIndex[VBRefBone.TargetRefSkelIndex];
 
-		VirtualBoneCompactPoseData.Add(FVirtualBoneCompactPoseData(FCompactPoseBoneIndex(VBInd), FCompactPoseBoneIndex(SourceInd), FCompactPoseBoneIndex(TargetInd)));
+		if ((VBInd != INDEX_NONE) && (SourceInd != INDEX_NONE) && (TargetInd != INDEX_NONE))
+		{
+			VirtualBoneCompactPoseData.Add(FVirtualBoneCompactPoseData(FCompactPoseBoneIndex(VBInd), FCompactPoseBoneIndex(SourceInd), FCompactPoseBoneIndex(TargetInd)));
+		}
 	}
 
 	// cache required curve UID list according to new bone sets
diff --git a/Engine/Source/Runtime/Engine/Private/ChartCreation.cpp b/Engine/Source/Runtime/Engine/Private/ChartCreation.cpp
--- a/Engine/Source/Runtime/Engine/Private/ChartCreation.cpp
+++ b/Engine/Source/Runtime/Engine/Private/ChartCreation.cpp
@@ -1099,9 +1099,7 @@ IPerformanceDataConsumer::FFrameData FPerformanceTrackingSystem::AnalyzeFrame(fl
 				// If our frame time is much larger than our last frame time, we'll count this as a hitch!
 				if (DeltaSeconds > (LastDeltaSeconds * HitchMultiplierAmount))
 				{
-					// We have a hitch!
-					GEngine->OnHitchDetectedDelegate.Broadcast(DeltaSeconds);
-
+			
 					// Check to see what we were limited by this frame
 					if (GGameThreadTime >= (MaxThreadTimeValue - EpsilonCycles))
 					{
@@ -1124,6 +1122,9 @@ IPerformanceDataConsumer::FFrameData FPerformanceTrackingSystem::AnalyzeFrame(fl
 						FrameData.HitchStatus = EFrameHitchType::UnknownUnit;
 					}
 
+					// We have a hitch!
+					GEngine->OnHitchDetectedDelegate.Broadcast(FrameData.HitchStatus, DeltaSeconds);
+
 					LastHitchTime = CurrentTime;
 				}
 			}
diff --git a/Engine/Source/Runtime/Engine/Private/Components/CharacterMovementComponent.cpp b/Engine/Source/Runtime/Engine/Private/Components/CharacterMovementComponent.cpp
--- a/Engine/Source/Runtime/Engine/Private/Components/CharacterMovementComponent.cpp
+++ b/Engine/Source/Runtime/Engine/Private/Components/CharacterMovementComponent.cpp
@@ -8630,28 +8630,9 @@ void UCharacterMovementComponent::ClientAdjustPosition_Implementation
 		WorldShiftedNewLocation = FRepMovement::RebaseOntoLocalOrigin(NewLocation, this);
 	}
 
-#if !UE_BUILD_SHIPPING
-	if (CharacterMovementCVars::NetShowCorrections != 0)
-	{
-		const FVector LocDiff = UpdatedComponent->GetComponentLocation() - WorldShiftedNewLocation;
-		const FString NewBaseString = NewBase ? NewBase->GetPathName(NewBase->GetOutermost()) : TEXT("None");
-		UE_LOG(LogNetPlayerMovement, Warning, TEXT("*** Client: Error for %s at Time=%.3f is %3.3f LocDiff(%s) ClientLoc(%s) ServerLoc(%s) NewBase: %s NewBone: %s ClientVel(%s) ServerVel(%s) SavedMoves %d"),
-			*GetNameSafe(CharacterOwner), TimeStamp, LocDiff.Size(), *LocDiff.ToString(), *UpdatedComponent->GetComponentLocation().ToString(), *WorldShiftedNewLocation.ToString(), *NewBaseString, *NewBaseBoneName.ToString(), *Velocity.ToString(), *NewVelocity.ToString(), ClientData->SavedMoves.Num());
-		const float DebugLifetime = CharacterMovementCVars::NetCorrectionLifetime;
-		DrawDebugCapsule(GetWorld(), UpdatedComponent->GetComponentLocation()	, CharacterOwner->GetSimpleCollisionHalfHeight(), CharacterOwner->GetSimpleCollisionRadius(), FQuat::Identity, FColor(255, 100, 100), true, DebugLifetime);
-		DrawDebugCapsule(GetWorld(), WorldShiftedNewLocation, CharacterOwner->GetSimpleCollisionHalfHeight(), CharacterOwner->GetSimpleCollisionRadius(), FQuat::Identity, FColor(100, 255, 100), true, DebugLifetime);
-	}
-#endif //!UE_BUILD_SHIPPING
 
-#if ROOT_MOTION_DEBUG
-	if (RootMotionSourceDebug::CVarDebugRootMotionSources.GetValueOnAnyThread() == 1)
-	{
-		const FVector VelocityCorrection = NewVelocity - Velocity;
-		FString AdjustedDebugString = FString::Printf(TEXT("PerformMovement ClientAdjustPosition_Implementation Velocity(%s) OldVelocity(%s) Correction(%s) TimeStamp(%f)"),
-			*NewVelocity.ToCompactString(), *Velocity.ToCompactString(), *VelocityCorrection.ToCompactString(), TimeStamp);
-		RootMotionSourceDebug::PrintOnScreen(*CharacterOwner, AdjustedDebugString);
-	}
-#endif
+	// Trigger event
+	OnClientCorrectionReceived(*ClientData, TimeStamp, NewLocation, NewVelocity, NewBase, NewBaseBoneName, bHasBase, bBaseRelativePosition, ServerMovementMode);
 
 	// Trust the server's positioning.
 	UpdatedComponent->SetWorldLocation(WorldShiftedNewLocation, false);
@@ -8703,6 +8684,34 @@ void UCharacterMovementComponent::ClientAdjustPosition_Implementation
 	ClientData->bUpdatePosition = true;
 }
 
+
+void UCharacterMovementComponent::OnClientCorrectionReceived(FNetworkPredictionData_Client_Character& ClientData, float TimeStamp, FVector NewLocation, FVector NewVelocity, UPrimitiveComponent* NewBase, FName NewBaseBoneName, bool bHasBase, bool bBaseRelativePosition, uint8 ServerMovementMode)
+{
+#if !UE_BUILD_SHIPPING
+	if (CharacterMovementCVars::NetShowCorrections != 0)
+	{
+		const FVector LocDiff = UpdatedComponent->GetComponentLocation() - NewLocation;
+		const FString NewBaseString = NewBase ? NewBase->GetPathName(NewBase->GetOutermost()) : TEXT("None");
+		UE_LOG(LogNetPlayerMovement, Warning, TEXT("*** Client: Error for %s at Time=%.3f is %3.3f LocDiff(%s) ClientLoc(%s) ServerLoc(%s) NewBase: %s NewBone: %s ClientVel(%s) ServerVel(%s) SavedMoves %d"),
+			   *GetNameSafe(CharacterOwner), TimeStamp, LocDiff.Size(), *LocDiff.ToString(), *UpdatedComponent->GetComponentLocation().ToString(), *NewLocation.ToString(), *NewBaseString, *NewBaseBoneName.ToString(), *Velocity.ToString(), *NewVelocity.ToString(), ClientData.SavedMoves.Num());
+		const float DebugLifetime = CharacterMovementCVars::NetCorrectionLifetime;
+		DrawDebugCapsule(GetWorld(), UpdatedComponent->GetComponentLocation(), CharacterOwner->GetSimpleCollisionHalfHeight(), CharacterOwner->GetSimpleCollisionRadius(), FQuat::Identity, FColor(255, 100, 100), true, DebugLifetime);
+		DrawDebugCapsule(GetWorld(), NewLocation, CharacterOwner->GetSimpleCollisionHalfHeight(), CharacterOwner->GetSimpleCollisionRadius(), FQuat::Identity, FColor(100, 255, 100), true, DebugLifetime);
+	}
+#endif //!UE_BUILD_SHIPPING
+
+#if ROOT_MOTION_DEBUG
+	if (RootMotionSourceDebug::CVarDebugRootMotionSources.GetValueOnAnyThread() == 1)
+	{
+		const FVector VelocityCorrection = NewVelocity - Velocity;
+		FString AdjustedDebugString = FString::Printf(TEXT("PerformMovement ClientAdjustPosition_Implementation Velocity(%s) OldVelocity(%s) Correction(%s) TimeStamp(%f)"),
+													  *NewVelocity.ToCompactString(), *Velocity.ToCompactString(), *VelocityCorrection.ToCompactString(), TimeStamp);
+		RootMotionSourceDebug::PrintOnScreen(*CharacterOwner, AdjustedDebugString);
+	}
+#endif
+}
+
+
 void UCharacterMovementComponent::ClientAdjustRootMotionPosition_Implementation(
 	float TimeStamp,
 	float ServerMontageTrackPosition,
diff --git a/Engine/Source/Runtime/Engine/Private/Components/PrimitiveComponent.cpp b/Engine/Source/Runtime/Engine/Private/Components/PrimitiveComponent.cpp
--- a/Engine/Source/Runtime/Engine/Private/Components/PrimitiveComponent.cpp
+++ b/Engine/Source/Runtime/Engine/Private/Components/PrimitiveComponent.cpp
@@ -1226,6 +1226,15 @@ void UPrimitiveComponent::SetCastShadow(bool NewCastShadow)
 	}
 }
 
+void UPrimitiveComponent::SetSingleSampleShadowFromStationaryLights(bool bNewSingleSampleShadowFromStationaryLights)
+{
+	if (bNewSingleSampleShadowFromStationaryLights != bSingleSampleShadowFromStationaryLights)
+	{
+		bSingleSampleShadowFromStationaryLights = bNewSingleSampleShadowFromStationaryLights;
+		MarkRenderStateDirty();
+	}
+}
+
 void UPrimitiveComponent::SetTranslucentSortPriority(int32 NewTranslucentSortPriority)
 {
 	if (NewTranslucentSortPriority != TranslucencySortPriority)
diff --git a/Engine/Source/Runtime/Engine/Private/Components/ReflectionCaptureComponent.cpp b/Engine/Source/Runtime/Engine/Private/Components/ReflectionCaptureComponent.cpp
--- a/Engine/Source/Runtime/Engine/Private/Components/ReflectionCaptureComponent.cpp
+++ b/Engine/Source/Runtime/Engine/Private/Components/ReflectionCaptureComponent.cpp
@@ -1489,7 +1489,7 @@ void UReflectionCaptureComponent::ReadbackFromGPU(UWorld* WorldToUpdate)
 		return;
 	}
 
-	if (bDerivedDataDirty && !IsRunningCommandlet() && WorldToUpdate && WorldToUpdate->FeatureLevel >= ERHIFeatureLevel::SM4)
+	if (bDerivedDataDirty && (!IsRunningCommandlet() || IsAllowCommandletRendering()) && WorldToUpdate && WorldToUpdate->FeatureLevel >= ERHIFeatureLevel::SM4)
 	{
 		FReflectionCaptureFullHDR* NewDerivedData = new FReflectionCaptureFullHDR();
 
diff --git a/Engine/Source/Runtime/Engine/Private/Components/SkeletalMeshComponent.cpp b/Engine/Source/Runtime/Engine/Private/Components/SkeletalMeshComponent.cpp
--- a/Engine/Source/Runtime/Engine/Private/Components/SkeletalMeshComponent.cpp
+++ b/Engine/Source/Runtime/Engine/Private/Components/SkeletalMeshComponent.cpp
@@ -1541,6 +1541,8 @@ void USkeletalMeshComponent::UpdateSlaveComponent()
 {
 	check (MasterPoseComponent.IsValid());
 
+	ResetMorphTargetCurves();
+
 	if (USkeletalMeshComponent* MasterSMC = Cast<USkeletalMeshComponent>(MasterPoseComponent.Get()))
 	{
 		// first set any animation-driven curves from the master SMC
diff --git a/Engine/Source/Runtime/Engine/Private/DataChannel.cpp b/Engine/Source/Runtime/Engine/Private/DataChannel.cpp
--- a/Engine/Source/Runtime/Engine/Private/DataChannel.cpp
+++ b/Engine/Source/Runtime/Engine/Private/DataChannel.cpp
@@ -603,7 +603,8 @@ bool UChannel::ReceivedNextBunch( FInBunch & Bunch, bool & bOutSkipAck )
 		{
 			if ( ChType != CHTYPE_Voice )	// Voice channels can open from both side simultaneously, so ignore this logic until we resolve this
 			{
-				check( !OpenedLocally );					// If we opened the channel, we shouldn't be receiving bOpen commands from the other side
+				// If we opened the channel, we shouldn't be receiving bOpen commands from the other side
+				checkf(!OpenedLocally, TEXT("Received channel open command for channel that was already opened locally. %s"), *Describe());
 				check( OpenPacketId.First == INDEX_NONE );	// This should be the first and only assignment of the packet range (we should only receive one bOpen bunch)
 				check( OpenPacketId.Last == INDEX_NONE );	// This should be the first and only assignment of the packet range (we should only receive one bOpen bunch)
 			}
diff --git a/Engine/Source/Runtime/Engine/Private/DemoNetDriver.cpp b/Engine/Source/Runtime/Engine/Private/DemoNetDriver.cpp
--- a/Engine/Source/Runtime/Engine/Private/DemoNetDriver.cpp
+++ b/Engine/Source/Runtime/Engine/Private/DemoNetDriver.cpp
@@ -379,7 +379,6 @@ bool UDemoNetDriver::InitBase( bool bInitAsClient, FNetworkNotify* InNotify, con
 		ViewerOverride					= nullptr;
 		bPrioritizeActors				= false;
 		bPauseRecording					= false;
-		CheckpointSaveMaxMSPerFrame		= 0.0f;
 
 		if ( RelevantTimeout == 0.0f )
 		{
diff --git a/Engine/Source/Runtime/Engine/Private/GameFramework/RootMotionSource.cpp b/Engine/Source/Runtime/Engine/Private/GameFramework/RootMotionSource.cpp
--- a/Engine/Source/Runtime/Engine/Private/GameFramework/RootMotionSource.cpp
+++ b/Engine/Source/Runtime/Engine/Private/GameFramework/RootMotionSource.cpp
@@ -1136,7 +1136,7 @@ void FRootMotionSource_JumpForce::PrepareRootMotion
 
 		// If we're beyond specified duration, we need to re-map times so that
 		// we continue our desired ending velocity
-		if (CurrentTimeFraction > 1.f)
+		if (TargetTimeFraction > 1.f)
 		{
 			float TimeFractionPastAllowable = TargetTimeFraction - 1.0f;
 			TargetTimeFraction -= TimeFractionPastAllowable;
@@ -1152,10 +1152,10 @@ void FRootMotionSource_JumpForce::PrepareRootMotion
 			TargetMoveFraction = TimeMappingCurve->GetFloatValue(TargetTimeFraction);
 		}
 
-		FVector CurrentRelativeLocation = GetRelativeLocation(CurrentMoveFraction);
-		FVector TargetRelativeLocation = GetRelativeLocation(TargetMoveFraction);
+		const FVector CurrentRelativeLocation = GetRelativeLocation(CurrentMoveFraction);
+		const FVector TargetRelativeLocation = GetRelativeLocation(TargetMoveFraction);
 
-		FVector Force = (TargetRelativeLocation - CurrentRelativeLocation) / MovementTickTime;
+		const FVector Force = (TargetRelativeLocation - CurrentRelativeLocation) / MovementTickTime;
 
 		// Debug
 #if ROOT_MOTION_DEBUG
@@ -1201,10 +1201,16 @@ void FRootMotionSource_JumpForce::PrepareRootMotion
 				GetTime(), GetTime() + SimulationTime, 
 				*CurrentLocation.ToString(), *CurrentTargetLocation.ToString(), 
 				*Force.ToString());
+
+			{
+				FString AdjustedDebugString = FString::Printf(TEXT("    FRootMotionSource_JumpForce::Prep Force(%s) SimTime(%.3f) MoveTime(%.3f) StartP(%.3f) EndP(%.3f)"),
+					*Force.ToCompactString(), SimulationTime, MovementTickTime, CurrentMoveFraction, TargetMoveFraction);
+				RootMotionSourceDebug::PrintOnScreen(Character, AdjustedDebugString);
+			}
 		}
 #endif
 
-		FTransform NewTransform(Force);
+		const FTransform NewTransform(Force);
 		RootMotionParams.Set(NewTransform);
 	}
 	else
diff --git a/Engine/Source/Runtime/Engine/Private/GameViewportClient.cpp b/Engine/Source/Runtime/Engine/Private/GameViewportClient.cpp
--- a/Engine/Source/Runtime/Engine/Private/GameViewportClient.cpp
+++ b/Engine/Source/Runtime/Engine/Private/GameViewportClient.cpp
@@ -748,40 +748,46 @@ EMouseCursor::Type UGameViewportClient::GetCursor(FViewport* InViewport, int32 X
 
 void UGameViewportClient::AddSoftwareCursor(EMouseCursor::Type Cursor, const FStringClassReference& CursorClass)
 {
-	if ( CursorClass.IsValid() )
+	if (ensureMsgf(CursorClass.IsValid(), TEXT("UGameViewportClient::AddCusor: Cursor class is not valid!")))
 	{
 		UClass* Class = CursorClass.TryLoadClass<UUserWidget>();
-		if ( Class )
+		if (Class)
 		{
 			UUserWidget* UserWidget = CreateWidget<UUserWidget>(GetGameInstance(), Class);
-			if ( ensure(UserWidget) )
+			if (UserWidget)
 			{
 				CursorWidgets.Add(Cursor, UserWidget->TakeWidget());
 			}
+			else
+			{
+				UE_LOG(LogPlayerManagement, Warning, TEXT("UGameViewportClient::AddCursor: Could not create cursor widget."));
+			}
 		}
 		else
 		{
-			FMessageLog("PIE").Error(FText::Format(LOCTEXT("CursorClassNotFoundFormat", "The cursor class '{0}' was not found, check your custom cursor settings."), FText::FromString(CursorClass.ToString())));
-			UE_LOG(LogPlayerManagement, Warning, TEXT("UGameViewportClient::AddCursor: The cursor class %s was not found, check your custom cursor settings."), *CursorClass.ToString());
+			UE_LOG(LogPlayerManagement, Warning, TEXT("UGameViewportClient::AddCursor: Could not load cursor class %s."), *CursorClass.GetAssetName());
 		}
 	}
-	else
-	{
-		UE_LOG(LogPlayerManagement, Warning, TEXT("UGameViewportClient::AddCursor: Attempting to add an invalid cursor class."));
-	}
 }
 
 TOptional<TSharedRef<SWidget>> UGameViewportClient::MapCursor(FViewport* InViewport, const FCursorReply& CursorReply)
 {
 	if (bUseSoftwareCursorWidgets)
 	{
-		const TSharedRef<SWidget>* CursorWidgetPtr = CursorWidgets.Find(CursorReply.GetCursorType());
-		if (CursorWidgetPtr != nullptr)
+		if (CursorReply.GetCursorType() != EMouseCursor::None)
 		{
-			return *CursorWidgetPtr;
+			const TSharedRef<SWidget>* CursorWidgetPtr = CursorWidgets.Find(CursorReply.GetCursorType());
+
+			if (CursorWidgetPtr != nullptr)
+			{
+				return *CursorWidgetPtr;
+			}
+			else
+			{
+				UE_LOG(LogPlayerManagement, Warning, TEXT("UGameViewportClient::MapCursor: Could not find cursor to map to %d."),int32(CursorReply.GetCursorType()));
+			}
 		}
 	}
-
 	return TOptional<TSharedRef<SWidget>>();
 }
 
diff --git a/Engine/Source/Runtime/Engine/Private/LevelTick.cpp b/Engine/Source/Runtime/Engine/Private/LevelTick.cpp
--- a/Engine/Source/Runtime/Engine/Private/LevelTick.cpp
+++ b/Engine/Source/Runtime/Engine/Private/LevelTick.cpp
@@ -1210,6 +1210,8 @@ void FTickableGameObject::TickObjects(UWorld* World, const int32 InTickType, con
  */
 void UWorld::Tick( ELevelTick TickType, float DeltaSeconds )
 {
+	SCOPE_TIME_GUARD(TEXT("UWorld::Tick"));
+
 	if (GIntraFrameDebuggingGameThread)
 	{
 		return;
@@ -1261,6 +1263,7 @@ void UWorld::Tick( ELevelTick TickType, float DeltaSeconds )
 
 	{
 		SCOPE_CYCLE_COUNTER(STAT_NetWorldTickTime);
+		SCOPE_TIME_GUARD(TEXT("UWorld::Tick - NetTick"));
 		// Update the net code and fetch all incoming packets.
 		BroadcastTickDispatch(DeltaSeconds);
 
@@ -1374,6 +1377,7 @@ void UWorld::Tick( ELevelTick TickType, float DeltaSeconds )
 
 			SCOPE_CYCLE_COUNTER(STAT_TickTime);
 			{
+				SCOPE_TIME_GUARD_MS(TEXT("UWorld::Tick - TG_PrePhysics"), 10);
 				SCOPE_CYCLE_COUNTER(STAT_TG_PrePhysics);
 				RunTickGroup(TG_PrePhysics);
 			}
@@ -1381,20 +1385,23 @@ void UWorld::Tick( ELevelTick TickType, float DeltaSeconds )
 			EnsureCollisionTreeIsBuilt();
 			bInTick = true;
 			{
-				SCOPE_CYCLE_COUNTER(STAT_TG_StartPhysics);
+				SCOPE_TIME_GUARD_MS(TEXT("UWorld::Tick - TG_StartPhysics"), 10);
 				RunTickGroup(TG_StartPhysics); 
 			}
 			{
 				SCOPE_CYCLE_COUNTER(STAT_TG_DuringPhysics);
+				SCOPE_TIME_GUARD_MS(TEXT("UWorld::Tick - TG_DuringPhysics"), 10);
 				RunTickGroup(TG_DuringPhysics, false); // No wait here, we should run until idle though. We don't care if all of the async ticks are done before we start running post-phys stuff
 			}
 			TickGroup = TG_EndPhysics; // set this here so the current tick group is correct during collision notifies, though I am not sure it matters. 'cause of the false up there^^^
 			{
 				SCOPE_CYCLE_COUNTER(STAT_TG_EndPhysics);
+				SCOPE_TIME_GUARD_MS(TEXT("UWorld::Tick - TG_EndPhysics"), 10);
 				RunTickGroup(TG_EndPhysics);
 			}
 			{
 				SCOPE_CYCLE_COUNTER(STAT_TG_PostPhysics);
+				SCOPE_TIME_GUARD_MS(TEXT("UWorld::Tick - TG_PostPhysics"), 10);
 				RunTickGroup(TG_PostPhysics);
 			}
 	
@@ -1428,11 +1435,15 @@ void UWorld::Tick( ELevelTick TickType, float DeltaSeconds )
 
 				if (TickType != LEVELTICK_TimeOnly && !bIsPaused)
 				{
+					SCOPE_TIME_GUARD_MS(TEXT("UWorld::Tick - TimerManager"), 5);
 					STAT(FScopeCycleCounter Context(GetTimerManager().GetStatId());)
 					GetTimerManager().Tick(DeltaSeconds);
 				}
 
-				FTickableGameObject::TickObjects(this, TickType, bIsPaused, DeltaSeconds);
+				{
+					SCOPE_TIME_GUARD_MS(TEXT("UWorld::Tick - TickObjects"), 5);
+					FTickableGameObject::TickObjects(this, TickType, bIsPaused, DeltaSeconds);
+				}
 			}
 
 			// Update cameras and streaming volumes
@@ -1473,10 +1484,12 @@ void UWorld::Tick( ELevelTick TickType, float DeltaSeconds )
 			SCOPE_CYCLE_COUNTER(STAT_TickTime);
 			{
 				SCOPE_CYCLE_COUNTER(STAT_TG_PostUpdateWork);
+				SCOPE_TIME_GUARD_MS(TEXT("UWorld::Tick - PostUpdateWork"), 5);
 				RunTickGroup(TG_PostUpdateWork);
 			}
 			{
 				SCOPE_CYCLE_COUNTER(STAT_TG_LastDemotable);
+				SCOPE_TIME_GUARD_MS(TEXT("UWorld::Tick - TG_LastDemotable"), 5);
 				RunTickGroup(TG_LastDemotable);
 			}
 
@@ -1496,6 +1509,7 @@ void UWorld::Tick( ELevelTick TickType, float DeltaSeconds )
 		// All tick is done, execute async trace
 		{
 			SCOPE_CYCLE_COUNTER(STAT_FinishAsyncTraceTickTime);
+			SCOPE_TIME_GUARD_MS(TEXT("UWorld::Tick - FinishAsyncTrace"), 5);
 			FinishAsyncTrace();
 		}
 	}
@@ -1526,6 +1540,7 @@ void UWorld::Tick( ELevelTick TickType, float DeltaSeconds )
 	// Tick the FX system.
 	if (!bIsPaused && FXSystem != NULL)
 	{
+		SCOPE_TIME_GUARD_MS(TEXT("UWorld::Tick - FX"), 5);
 		FXSystem->Tick(DeltaSeconds);
 	}
 	
diff --git a/Engine/Source/Runtime/Engine/Private/SkeletalMesh.cpp b/Engine/Source/Runtime/Engine/Private/SkeletalMesh.cpp
--- a/Engine/Source/Runtime/Engine/Private/SkeletalMesh.cpp
+++ b/Engine/Source/Runtime/Engine/Private/SkeletalMesh.cpp
@@ -5158,7 +5158,7 @@ public:
 			if (ActualPreviewSectionIdx != INDEX_NONE && Sections.IsValidIndex(ActualPreviewSectionIdx))
 			{
 				const FSkelMeshSection& PreviewSection = Sections[ActualPreviewSectionIdx];
-				if (PreviewSection.CorrespondClothSectionIndex != INDEX_NONE)
+				if (PreviewSection.bDisabled && PreviewSection.CorrespondClothSectionIndex != INDEX_NONE)
 				{
 					ActualPreviewSectionIdx = PreviewSection.CorrespondClothSectionIndex;
 				}
@@ -5181,7 +5181,7 @@ public:
 				if (ActualPreviewSectionIdx != INDEX_NONE)
 				{
 					const FSkelMeshSection& PreviewSection = Sections[ActualPreviewSectionIdx];
-					if (PreviewSection.CorrespondClothSectionIndex != INDEX_NONE)
+					if (PreviewSection.bDisabled && PreviewSection.CorrespondClothSectionIndex != INDEX_NONE)
 					{
 						ActualPreviewSectionIdx = PreviewSection.CorrespondClothSectionIndex;
 					}
@@ -5287,6 +5287,19 @@ void FSkeletalMeshSceneProxy::GetMeshElementsConditionallySelectable(const TArra
 
 		check(LODSection.SectionElements.Num() == LODModel.Sections.Num());
 
+#if WITH_EDITORONLY_DATA
+		//Find the real editor selected section
+		int32 RealSelectedEditorSection = SkeletalMeshForDebug->SelectedEditorSection;
+		if (RealSelectedEditorSection != INDEX_NONE && LODModel.Sections.IsValidIndex(SkeletalMeshForDebug->SelectedEditorSection))
+		{
+			const FSkelMeshSection& SelectEditorSection = LODModel.Sections[SkeletalMeshForDebug->SelectedEditorSection];
+			if (SelectEditorSection.bDisabled && SelectEditorSection.CorrespondClothSectionIndex != INDEX_NONE && LODModel.Sections.IsValidIndex(SelectEditorSection.CorrespondClothSectionIndex))
+			{
+				RealSelectedEditorSection = SelectEditorSection.CorrespondClothSectionIndex;
+			}
+		}
+#endif
+
 		for (FSkeletalMeshSectionIter Iter(LODIndex, *MeshObject, LODModel, LODSection); Iter; ++Iter)
 		{
 			const FSkelMeshSection& Section = Iter.GetSection();
@@ -5304,7 +5317,7 @@ void FSkeletalMeshSceneProxy::GetMeshElementsConditionallySelectable(const TArra
 			}
 			else
 			{
-				bSectionSelected = (SkeletalMeshForDebug->SelectedEditorSection == Iter.GetSectionElementIndex());
+				bSectionSelected = (RealSelectedEditorSection == SectionIndex);
 			}
 			
 #endif
diff --git a/Engine/Source/Runtime/Engine/Private/UnrealEngine.cpp b/Engine/Source/Runtime/Engine/Private/UnrealEngine.cpp
--- a/Engine/Source/Runtime/Engine/Private/UnrealEngine.cpp
+++ b/Engine/Source/Runtime/Engine/Private/UnrealEngine.cpp
@@ -191,6 +191,7 @@
 
 #include "GeneralProjectSettings.h"
 #include "ProfilingDebugging/LoadTimeTracker.h"
+#include "ObjectKey.h"
 #include "AssetRegistryModule.h"
 
 DEFINE_LOG_CATEGORY(LogEngine);
@@ -4328,7 +4329,8 @@ bool UEngine::HandleMemReportDeferredCommand( const TCHAR* Cmd, FOutputDevice& A
 		UE_LOG(LogEngine, Log, TEXT("MemReportDeferred: saving to %s"), *FilenameFull);		
 	}
 
-	ReportAr->Logf( TEXT( "CommandLine Options: %s" ) LINE_TERMINATOR, FCommandLine::Get() );
+	ReportAr->Logf(TEXT( "CommandLine Options: %s" ),  FCommandLine::Get() );
+	ReportAr->Logf(TEXT("Time Since Boot: %.02f Seconds") LINE_TERMINATOR, FPlatformTime::Seconds() - GStartTime);
 
 	// Run commands from the ini
 	FConfigSection* CommandsToRun = GConfig->GetSectionPrivate(TEXT("MemReportCommands"), 0, 1, GEngineIni);
@@ -5424,6 +5426,24 @@ bool UEngine::HandleObjCommand( const TCHAR* Cmd, FOutputDevice& Ar )
 	}
 	else if( FParse::Command(&Cmd,TEXT("LIST")) )
 	{
+		static TSet<FObjectKey> ForgottenObjects;
+
+		// "obj list forget" will prevent all current objects from being reported in future "obj list" commands.
+		// "obj list remember" clears that list
+		if (FParse::Command(&Cmd, TEXT("FORGET")))
+		{
+			for (FObjectIterator It; It; ++It)
+			{
+				ForgottenObjects.Add(FObjectKey(*It));
+			}
+			return true;
+		}
+		else if (FParse::Command(&Cmd, TEXT("REMEMBER")))
+		{
+			ForgottenObjects.Empty();
+			return true;
+		}
+
 		FString CmdLineOut = FString::Printf(TEXT("Obj List: %s"), Cmd);
 		Ar.Log( *CmdLineOut );
 		Ar.Log( TEXT("Objects:") );
@@ -5481,6 +5501,10 @@ bool UEngine::HandleObjCommand( const TCHAR* Cmd, FOutputDevice& Ar )
 
 			for( FObjectIterator It; It; ++It )
 			{
+				if (ForgottenObjects.Contains(FObjectKey(*It)))
+				{
+					continue;
+				}
 				if (It->IsTemplate(RF_ClassDefaultObject))
 				{
 					if( !bShouldIncludeDefaultObjects )
diff --git a/Engine/Source/Runtime/Engine/Public/AnimationUtils.h b/Engine/Source/Runtime/Engine/Public/AnimationUtils.h
--- a/Engine/Source/Runtime/Engine/Public/AnimationUtils.h
+++ b/Engine/Source/Runtime/Engine/Public/AnimationUtils.h
@@ -104,9 +104,6 @@ enum EConvertToAdditive
 	CTA_MAX	// Always last one
 };
 
-/** Control animation recompression upon load. */
-extern ENGINE_API bool GDisableAnimationRecompression;
-
 /**
  * A collection of useful functions for skeletal mesh animation.
  */
diff --git a/Engine/Source/Runtime/Engine/Public/BonePose.h b/Engine/Source/Runtime/Engine/Public/BonePose.h
--- a/Engine/Source/Runtime/Engine/Public/BonePose.h
+++ b/Engine/Source/Runtime/Engine/Public/BonePose.h
@@ -796,18 +796,21 @@ void FCSPose<PoseType>::ConvertToLocalPoses(PoseType& OutPose) const
 	// that doesn't mean this local has to change
 	// go from child to parent since I need parent inverse to go back to local
 	// root is same, so no need to do Index == 0
-	for (const BoneIndexType BoneIndex : Pose.ForEachBoneIndexReverse())
+	const BoneIndexType RootBoneIndex(0);
+	if (ComponentSpaceFlags[RootBoneIndex])
 	{
-		if (!BoneIndex.IsRootBone())
-		{
-			// root is already verified, so root should not come here
-			if (ComponentSpaceFlags[BoneIndex])
-			{
-				const FCompactPoseBoneIndex ParentIndex = Pose.GetParentBoneIndex(BoneIndex);
+		OutPose[RootBoneIndex] = Pose[RootBoneIndex];
+	}
 
-				OutPose[BoneIndex].SetToRelativeTransform(OutPose[ParentIndex]);
-				OutPose[BoneIndex].NormalizeRotation();
-			}
+	const int32 NumBones = Pose.GetNumBones();
+	for (int32 Index = NumBones - 1; Index > 0; Index--)
+	{
+		const BoneIndexType BoneIndex(Index);
+		if (ComponentSpaceFlags[BoneIndex])
+		{
+			const BoneIndexType ParentIndex = Pose.GetParentBoneIndex(BoneIndex);
+			OutPose[BoneIndex].SetToRelativeTransform(OutPose[ParentIndex]);
+			OutPose[BoneIndex].NormalizeRotation();
 		}
 	}
 }
\ No newline at end of file
diff --git a/Engine/Source/Runtime/Engine/Public/ReferenceSkeleton.h b/Engine/Source/Runtime/Engine/Public/ReferenceSkeleton.h
--- a/Engine/Source/Runtime/Engine/Public/ReferenceSkeleton.h
+++ b/Engine/Source/Runtime/Engine/Public/ReferenceSkeleton.h
@@ -153,9 +153,13 @@ private:
 
 	int32 GetParentIndexInternal(const int32 BoneIndex, const TArray<FMeshBoneInfo>& BoneInfo) const
 	{
-		// Parent must be valid. Either INDEX_NONE for Root, or before us.
-		checkSlow(((BoneIndex == 0) && (BoneInfo[BoneIndex].ParentIndex == INDEX_NONE)) || ((BoneIndex > 0) && BoneInfo.IsValidIndex(BoneInfo[BoneIndex].ParentIndex)));
-		return BoneInfo[BoneIndex].ParentIndex;
+		const int32 ParentIndex = BoneInfo[BoneIndex].ParentIndex;
+
+		// Parent must be valid. Either INDEX_NONE for Root, or before children for non root bones.
+		checkSlow(((BoneIndex == 0) && (ParentIndex == INDEX_NONE))
+			|| ((BoneIndex > 0) && BoneInfo.IsValidIndex(ParentIndex) && (ParentIndex < BoneIndex)));
+
+		return ParentIndex;
 	}
 
 	void UpdateRefPoseTransform(const int32 BoneIndex, const FTransform& BonePose)
@@ -180,8 +184,9 @@ private:
 		// Normalize Quaternion to be safe.
 		RawRefBonePose[BoneIndex].NormalizeRotation();
 
-		// Parent must be valid. Either INDEX_NONE for Root, or before us.
-		check(((BoneIndex == 0) && (BoneInfo.ParentIndex == INDEX_NONE)) || ((BoneIndex > 0) && RawRefBoneInfo.IsValidIndex(BoneInfo.ParentIndex)));
+		// Parent must be valid. Either INDEX_NONE for Root, or before children for non root bones.
+		check(((BoneIndex == 0) && (BoneInfo.ParentIndex == INDEX_NONE)) 
+			|| ((BoneIndex > 0) && RawRefBoneInfo.IsValidIndex(BoneInfo.ParentIndex) && (BoneInfo.ParentIndex < BoneIndex)));
 	}
 
 	// Help us translate a virtual bone source into a raw bone source (for evaluating virtual bone transform)
diff --git a/Engine/Source/Runtime/GameplayTags/Private/GameplayTagsManager.cpp b/Engine/Source/Runtime/GameplayTags/Private/GameplayTagsManager.cpp
--- a/Engine/Source/Runtime/GameplayTags/Private/GameplayTagsManager.cpp
+++ b/Engine/Source/Runtime/GameplayTags/Private/GameplayTagsManager.cpp
@@ -24,6 +24,7 @@
 #include "PropertyHandle.h"
 FSimpleMulticastDelegate UGameplayTagsManager::OnEditorRefreshGameplayTagTree;
 #endif
+#include "IConsoleManager.h"
 
 
 #define LOCTEXT_NAMESPACE "GameplayTagManager"
@@ -274,11 +275,21 @@ void UGameplayTagsManager::ConstructGameplayTagTree()
 				FGameplayTag OldTag = RequestGameplayTag(OldTagName, false); //< This only succeeds if OldTag is in the Table!
 				if (OldTag.IsValid())
 				{
-					UE_LOG(LogGameplayTags, Warning,
-						TEXT("Old tag (%s) which is being redirected still exists in the table!  Generally you should "
+					FGameplayTagContainer MatchingChildren = RequestGameplayTagChildren(OldTag);
+
+					FString Msg = FString::Printf(TEXT("Old tag (%s) which is being redirected still exists in the table!  Generally you should "
 						TEXT("remove the old tags from the table when you are redirecting to new tags, or else users will ")
-						TEXT("still be able to add the old tags to containers.")), *OldTagName.ToString()
-						);
+						TEXT("still be able to add the old tags to containers.")), *OldTagName.ToString());
+
+					if (MatchingChildren.Num() == 0)
+					{
+						UE_LOG(LogGameplayTags, Warning, TEXT("%s"), *Msg);
+					}
+					else
+					{
+						Msg += TEXT("\nSuppressed warning due to redirected tag being a single component that matched other hierarchy elements.");
+						UE_LOG(LogGameplayTags, Log, TEXT("%s"), *Msg);
+					}
 				}
 
 				FGameplayTag NewTag = (NewTagName != NAME_None) ? RequestGameplayTag(NewTagName, false) : FGameplayTag();
@@ -320,6 +331,8 @@ void UGameplayTagsManager::ConstructGameplayTagTree()
 	}
 }
 
+int32 PrintNetIndiceAssignment = 0;
+static FAutoConsoleVariableRef CVarPrintNetIndiceAssignment(TEXT("GameplayTags.PrintNetIndiceAssignment"), PrintNetIndiceAssignment, TEXT("Logs GameplayTag NetIndice assignment"), ECVF_Default );
 void UGameplayTagsManager::ConstructNetIndex()
 {
 	NetworkGameplayTagNodeIndex.Empty();
@@ -365,11 +378,21 @@ void UGameplayTagsManager::ConstructNetIndex()
 		NetworkGameplayTagNodeIndex.SetNum(INVALID_TAGNETINDEX - 1);
 	}
 
+
+
+	UE_CLOG(PrintNetIndiceAssignment, LogGameplayTags, Display, TEXT("Assigning NetIndices to %d tags."), NetworkGameplayTagNodeIndex.Num() );
+
 	for (FGameplayTagNetIndex i = 0; i < NetworkGameplayTagNodeIndex.Num(); i++)
 	{
 		if (NetworkGameplayTagNodeIndex[i].IsValid())
 		{
 			NetworkGameplayTagNodeIndex[i]->NetIndex = i;
+
+			UE_CLOG(PrintNetIndiceAssignment, LogGameplayTags, Display, TEXT("Assigning NetIndex (%d) to Tag (%s)"), i, *NetworkGameplayTagNodeIndex[i]->GetCompleteTag().ToString());
+		}
+		else
+		{
+			UE_LOG(LogGameplayTags, Warning, TEXT("TagNode Indice %d is invalid!"), i);
 		}
 	}
 }
diff --git a/Engine/Source/Runtime/GameplayTasks/Classes/GameplayTask.h b/Engine/Source/Runtime/GameplayTasks/Classes/GameplayTask.h
--- a/Engine/Source/Runtime/GameplayTasks/Classes/GameplayTask.h
+++ b/Engine/Source/Runtime/GameplayTasks/Classes/GameplayTask.h
@@ -135,7 +135,7 @@ public:
 	FString GetDebugDescription() const;
 };
 
-UCLASS(Abstract, meta = (ExposedAsyncProxy), config = Game)
+UCLASS(Abstract, BlueprintType, meta = (ExposedAsyncProxy), config = Game)
 class GAMEPLAYTASKS_API UGameplayTask : public UObject, public IGameplayTaskOwnerInterface
 {
 	GENERATED_BODY()
diff --git a/Engine/Source/Runtime/Launch/Launch.Build.cs b/Engine/Source/Runtime/Launch/Launch.Build.cs
--- a/Engine/Source/Runtime/Launch/Launch.Build.cs
+++ b/Engine/Source/Runtime/Launch/Launch.Build.cs
@@ -103,6 +103,7 @@ public class Launch : ModuleRules
 				new string[] {
 					"NetworkFile",
 					"StreamingFile",
+					"CookedIterativeFile",
     				"AutomationWorker",
 				}
 			);
diff --git a/Engine/Source/Runtime/Launch/Private/LaunchEngineLoop.cpp b/Engine/Source/Runtime/Launch/Private/LaunchEngineLoop.cpp
--- a/Engine/Source/Runtime/Launch/Private/LaunchEngineLoop.cpp
+++ b/Engine/Source/Runtime/Launch/Private/LaunchEngineLoop.cpp
@@ -623,9 +623,20 @@ bool LaunchCheckForFileOverride(const TCHAR* CmdLine, bool& OutFileOverrideFound
 			FPlatformFileManager::Get().SetPlatformFile(*CurrentPlatformFile);
 		}
 
+		bool bShouldUseCookedIterativeFile = false;
+		if ( !bShouldUseStreamingFile && !NetworkPlatformFile )
+		{
+			NetworkPlatformFile = ConditionallyCreateFileWrapper(TEXT("CookedIterativeFile"), CurrentPlatformFile, CmdLine, &bNetworkFailedToInitialize, &bShouldUseCookedIterativeFile);
+			if (NetworkPlatformFile)
+			{
+				CurrentPlatformFile = NetworkPlatformFile;
+				FPlatformFileManager::Get().SetPlatformFile(*CurrentPlatformFile);
+			}
+		}
+
 		// if streaming network platform file was tried this loop don't try this one
 		// Network file wrapper (only create if the streaming wrapper hasn't been created)
-		if ( !bShouldUseStreamingFile && !NetworkPlatformFile)
+		if ( !bShouldUseStreamingFile && !bShouldUseCookedIterativeFile && !NetworkPlatformFile)
 		{
 			NetworkPlatformFile = ConditionallyCreateFileWrapper(TEXT("NetworkFile"), CurrentPlatformFile, CmdLine, &bNetworkFailedToInitialize);
 			if (NetworkPlatformFile)
diff --git a/Engine/Source/Runtime/LevelSequence/Private/LevelSequencePlayer.cpp b/Engine/Source/Runtime/LevelSequence/Private/LevelSequencePlayer.cpp
--- a/Engine/Source/Runtime/LevelSequence/Private/LevelSequencePlayer.cpp
+++ b/Engine/Source/Runtime/LevelSequence/Private/LevelSequencePlayer.cpp
@@ -260,8 +260,10 @@ void ULevelSequencePlayer::TakeFrameSnapshot(FLevelSequencePlayerSnapshot& OutSn
 		return;
 	}
 
+	const float StartTimeWithoutWarmupFrames = SnapshotOffsetTime.IsSet() ? StartTime + SnapshotOffsetTime.GetValue() : StartTime;
+
 	// Use the actual last evaluation time as per the play position, which accounts for fixed time step offsetting
-	const float CurrentTime = StartTime + PlayPosition.GetLastPlayEvalPostition().Get(TimeCursorPosition);
+	const float CurrentTime = StartTimeWithoutWarmupFrames + PlayPosition.GetLastPlayEvalPostition().Get(TimeCursorPosition);
 
 	OutSnapshot.Settings = SnapshotSettings;
 
diff --git a/Engine/Source/Runtime/LevelSequence/Public/LevelSequencePlayer.h b/Engine/Source/Runtime/LevelSequence/Public/LevelSequencePlayer.h
--- a/Engine/Source/Runtime/LevelSequence/Public/LevelSequencePlayer.h
+++ b/Engine/Source/Runtime/LevelSequence/Public/LevelSequencePlayer.h
@@ -150,6 +150,9 @@ public:
 	/** Take a snapshot of the current state of this player */
 	void TakeFrameSnapshot(FLevelSequencePlayerSnapshot& OutSnapshot) const;
 
+	/** Set the offset time for the snapshot */
+	void SetSnapshotOffsetTime(float InTime) {SnapshotOffsetTime = TOptional<float>(InTime); }
+
 private:
 
 	/** Add tick prerequisites so that the level sequence actor ticks before all the actors it controls */
@@ -175,6 +178,8 @@ protected:
 	/** How to take snapshots */
 	FLevelSequenceSnapshotSettings SnapshotSettings;
 
+	TOptional<float> SnapshotOffsetTime;
+
 	TWeakObjectPtr<UCameraComponent> CachedCameraComponent;
 
 	/** Array of additional event receivers */
diff --git a/Engine/Source/Runtime/MovieSceneCapture/Private/AutomatedLevelSequenceCapture.cpp b/Engine/Source/Runtime/MovieSceneCapture/Private/AutomatedLevelSequenceCapture.cpp
--- a/Engine/Source/Runtime/MovieSceneCapture/Private/AutomatedLevelSequenceCapture.cpp
+++ b/Engine/Source/Runtime/MovieSceneCapture/Private/AutomatedLevelSequenceCapture.cpp
@@ -390,6 +390,9 @@ void UAutomatedLevelSequenceCapture::SetupFrameRange()
 					(float)PlaybackStartFrame / (float)Settings.FrameRate,
 					(float)PlaybackEndFrame / (float)Settings.FrameRate );
 				Actor->SequencePlayer->SetPlaybackPosition(0.f);
+
+				const float WarmupTime = WarmUpFrameCount / CachedState.Settings.FrameRate;
+				Actor->SequencePlayer->SetSnapshotOffsetTime(WarmupTime);
 			}
 		}
 	}
diff --git a/Engine/Source/Runtime/NetworkFile/Private/NetworkPlatformFile.cpp b/Engine/Source/Runtime/NetworkFile/Private/NetworkPlatformFile.cpp
--- a/Engine/Source/Runtime/NetworkFile/Private/NetworkPlatformFile.cpp
+++ b/Engine/Source/Runtime/NetworkFile/Private/NetworkPlatformFile.cpp
@@ -42,6 +42,9 @@ FNetworkPlatformFile::FNetworkPlatformFile()
 	: bHasLoadedDDCDirectories(false)
 	, InnerPlatformFile(NULL)
 	, bIsUsable(false)
+	, FileServerPort(0)
+	, ConnectionFlags(EConnectionFlags::None)
+	, HeartbeatFrequency(5.0f)
 	, FinishedAsyncNetworkReadUnsolicitedFiles(NULL)
 	, FinishedAsyncWriteUnsolicitedFiles(NULL)
 	, Transport(NULL)
@@ -199,7 +202,7 @@ void FNetworkPlatformFile::InitializeAfterSetActive()
 
 		// send the filenames and timestamps to the server
 		FNetworkFileArchive Payload(NFS_Messages::GetFileList);
-		FillGetFileList(Payload, false);
+		FillGetFileList(Payload);
 
 		// send the directories over, and wait for a response
 		FArrayReader Response;
@@ -214,169 +217,178 @@ void FNetworkPlatformFile::InitializeAfterSetActive()
 			int32 ServerPackageVersion = 0;
 			int32 ServerPackageLicenseeVersion = 0;
 			ProcessServerInitialResponse(Response, ServerPackageVersion, ServerPackageLicenseeVersion);
+			ProcessServerCachedFilesResponse(Response, ServerPackageVersion, ServerPackageLicenseeVersion);
 
+			// make sure we can sync a file
+			FString TestSyncFile = FPaths::Combine(*(FPaths::EngineDir()), TEXT("Config/BaseEngine.ini"));
 
-			/* The server root content directories */
-			TArray<FString> ServerRootContentDirectories;
-			Response << ServerRootContentDirectories;
-
-			// receive a list of the cache files and their timestamps
-			TMap<FString, FDateTime> ServerCachedFiles;
-			Response << ServerCachedFiles;
-
-			bool bDeleteAllFiles = true;
-			// Check the stored cooked version
-			FString CookedVersionFile = FPaths::GeneratedConfigDir() / TEXT("CookedVersion.txt");
-
-			if (InnerPlatformFile->FileExists(*CookedVersionFile) == true)
-			{
-				IFileHandle* FileHandle = InnerPlatformFile->OpenRead(*CookedVersionFile);
-				if (FileHandle != NULL)
-				{
-					int32 StoredPackageCookedVersion;
-					int32 StoredPackageCookedLicenseeVersion;
-					if (FileHandle->Read((uint8*)&StoredPackageCookedVersion, sizeof(int32)) == true)
-					{
-						if (FileHandle->Read((uint8*)&StoredPackageCookedLicenseeVersion, sizeof(int32)) == true)
-						{
-							if ((ServerPackageVersion == StoredPackageCookedVersion) &&
-								(ServerPackageLicenseeVersion == StoredPackageCookedLicenseeVersion))
-							{
-								bDeleteAllFiles = false;
-							}
-							else
-							{
-								UE_LOG(LogNetworkPlatformFile, Display, 
-									TEXT("Engine version mismatch: Server %d.%d, Stored %d.%d\n"), 
-									ServerPackageVersion, ServerPackageLicenseeVersion,
-									StoredPackageCookedVersion, StoredPackageCookedLicenseeVersion);
-							}
-						}
-					}
-
-					delete FileHandle;
-				}
-			}
-			else
+			InnerPlatformFile->SetReadOnly(*TestSyncFile, false);
+			InnerPlatformFile->DeleteFile(*TestSyncFile);
+			if (InnerPlatformFile->FileExists(*TestSyncFile))
 			{
-				UE_LOG(LogNetworkPlatformFile, Display, TEXT("Cooked version file missing: %s\n"), *CookedVersionFile);
+				UE_LOG(LogNetworkPlatformFile, Fatal, TEXT("Could not delete file sync test file %s."), *TestSyncFile);
 			}
 
-			if (bDeleteAllFiles == true)
+			EnsureFileIsLocal(TestSyncFile);
+
+			if (!InnerPlatformFile->FileExists(*TestSyncFile) || InnerPlatformFile->FileSize(*TestSyncFile) < 1)
 			{
-				// Make sure the config file exists...
-				InnerPlatformFile->CreateDirectoryTree(*(FPaths::GeneratedConfigDir()));
-				// Update the cooked version file
-				IFileHandle* FileHandle = InnerPlatformFile->OpenWrite(*CookedVersionFile);
-				if (FileHandle != NULL)
-				{
-					FileHandle->Write((const uint8*)&ServerPackageVersion, sizeof(int32));
-					FileHandle->Write((const uint8*)&ServerPackageLicenseeVersion, sizeof(int32));
-					delete FileHandle;
-				}
+				UE_LOG(LogNetworkPlatformFile, Fatal, TEXT("Could not sync test file %s."), *TestSyncFile);
 			}
+		}
+	}
 
-			// list of directories to skip
-			TArray<FString> DirectoriesToSkip;
-			TArray<FString> DirectoriesToNotRecurse;
-			// use the timestamp grabbing visitor to get all the content times
-			FLocalTimestampDirectoryVisitor Visitor(*InnerPlatformFile, DirectoriesToSkip, DirectoriesToNotRecurse, false);
+	FPlatformMisc::LowLevelOutputDebugStringf(TEXT("Network file startup time: %5.3f seconds\n"), NetworkFileStartupTime);
 
-			/*TArray<FString> RootContentPaths;
-			FPackageName::QueryRootContentPaths(RootContentPaths); */
-			for (TArray<FString>::TConstIterator RootPathIt(ServerRootContentDirectories); RootPathIt; ++RootPathIt)
-			{
-				/*const FString& RootPath = *RootPathIt;
-				const FString& ContentFolder = FPackageName::LongPackageNameToFilename(RootPath);*/
-				const FString& ContentFolder = *RootPathIt;
-				InnerPlatformFile->IterateDirectory( *ContentFolder, Visitor);
-			}
+}
 
-			// delete out of date files using the server cached files
-			for (TMap<FString, FDateTime>::TIterator It(ServerCachedFiles); It; ++It)
-			{
-				bool bDeleteFile = bDeleteAllFiles;
-				FString ServerFile = It.Key();
+void FNetworkPlatformFile::ProcessServerCachedFilesResponse(FArrayReader& Response, const int32 ServerPackageVersion, const int32 ServerPackageLicenseeVersion)
+{
+	/* The server root content directories */
+	TArray<FString> ServerRootContentDirectories;
+	Response << ServerRootContentDirectories;
 
-				// Convert the filename to the client version
-				ConvertServerFilenameToClientFilename(ServerFile);
+	// receive a list of the cache files and their timestamps
+	TMap<FString, FDateTime> ServerCachedFiles;
+	Response << ServerCachedFiles;
 
-				// Set it in the visitor file times list
-				// If there is any pathing difference (relative path, or whatever) between the server's filelist and the results
-				// of platform directory iteration then this will Add a new entry rather than override the existing one.  This causes local file deletes
-				// and longer loads as we will never see the benefits of local device caching.
-				Visitor.FileTimes.Add(ServerFile, FDateTime::MinValue());
+	bool bDeleteAllFiles = true;
+	// Check the stored cooked version
+	FString CookedVersionFile = FPaths::GeneratedConfigDir() / TEXT("CookedVersion.txt");
 
-				if (bDeleteFile == false)
+	if (InnerPlatformFile->FileExists(*CookedVersionFile) == true)
+	{
+		IFileHandle* FileHandle = InnerPlatformFile->OpenRead(*CookedVersionFile);
+		if (FileHandle != NULL)
+		{
+			int32 StoredPackageCookedVersion;
+			int32 StoredPackageCookedLicenseeVersion;
+			if (FileHandle->Read((uint8*)&StoredPackageCookedVersion, sizeof(int32)) == true)
+			{
+				if (FileHandle->Read((uint8*)&StoredPackageCookedLicenseeVersion, sizeof(int32)) == true)
 				{
-					// Check the time stamps...
-					// get local time
-					FDateTime LocalTime = InnerPlatformFile->GetTimeStamp(*ServerFile);
-					// If local time == MinValue than the file does not exist in the cache.
-					if (LocalTime != FDateTime::MinValue())
+					if ((ServerPackageVersion == StoredPackageCookedVersion) &&
+						(ServerPackageLicenseeVersion == StoredPackageCookedLicenseeVersion))
 					{
-						FDateTime ServerTime = It.Value();
-						// delete if out of date
-						// We will use 1.0 second as the tolerance to cover any platform differences in resolution
-						FTimespan TimeDiff = LocalTime - ServerTime;
-						double TimeDiffInSeconds = TimeDiff.GetTotalSeconds();
-						bDeleteFile = (TimeDiffInSeconds > 1.0) || (TimeDiffInSeconds < -1.0);
-						if (bDeleteFile == true)
-						{
-							if (InnerPlatformFile->FileExists(*ServerFile) == true)
-							{
-								UE_LOG(LogNetworkPlatformFile, Display, TEXT("Deleting cached file: TimeDiff %5.3f, %s"), TimeDiffInSeconds, *It.Key());
-							}
-							else
-							{
-								// It's a directory
-								bDeleteFile = false;
-							}
-						}
-						else
-						{
-							UE_LOG(LogNetworkPlatformFile, Display, TEXT("Keeping cached file: %s, TimeDiff worked out ok"), *ServerFile);
-						}
+						bDeleteAllFiles = false;
+					}
+					else
+					{
+						UE_LOG(LogNetworkPlatformFile, Display,
+							TEXT("Engine version mismatch: Server %d.%d, Stored %d.%d\n"),
+							ServerPackageVersion, ServerPackageLicenseeVersion,
+							StoredPackageCookedVersion, StoredPackageCookedLicenseeVersion);
 					}
-				}
-				if (bDeleteFile == true)
-				{
-					InnerPlatformFile->DeleteFile(*ServerFile);
 				}
 			}
 
-			// Any content files we have locally that were not cached, delete them
-			for (TMap<FString, FDateTime>::TIterator It(Visitor.FileTimes); It; ++It)
-			{
-				if (It.Value() != FDateTime::MinValue())
-				{
-					// This was *not* found in the server file list... delete it
-					UE_LOG(LogNetworkPlatformFile, Display, TEXT("Deleting cached file: %s"), *It.Key());
-					InnerPlatformFile->DeleteFile(*It.Key());
-				}
-			}
+			delete FileHandle;
+		}
+	}
+	else
+	{
+		UE_LOG(LogNetworkPlatformFile, Display, TEXT("Cooked version file missing: %s\n"), *CookedVersionFile);
+	}
 
-			// make sure we can sync a file
-			FString TestSyncFile = FPaths::Combine(*(FPaths::EngineDir()), TEXT("Config/BaseEngine.ini"));
+	if (bDeleteAllFiles == true)
+	{
+		// Make sure the config file exists...
+		InnerPlatformFile->CreateDirectoryTree(*(FPaths::GeneratedConfigDir()));
+		// Update the cooked version file
+		IFileHandle* FileHandle = InnerPlatformFile->OpenWrite(*CookedVersionFile);
+		if (FileHandle != NULL)
+		{
+			FileHandle->Write((const uint8*)&ServerPackageVersion, sizeof(int32));
+			FileHandle->Write((const uint8*)&ServerPackageLicenseeVersion, sizeof(int32));
+			delete FileHandle;
+		}
+	}
 
-			InnerPlatformFile->SetReadOnly(*TestSyncFile, false);
-			InnerPlatformFile->DeleteFile(*TestSyncFile);
-			if (InnerPlatformFile->FileExists(*TestSyncFile))
-			{
-				UE_LOG(LogNetworkPlatformFile, Fatal, TEXT("Could not delete file sync test file %s."), *TestSyncFile);
-			}
+	// list of directories to skip
+	TArray<FString> DirectoriesToSkip;
+	TArray<FString> DirectoriesToNotRecurse;
+	// use the timestamp grabbing visitor to get all the content times
+	FLocalTimestampDirectoryVisitor Visitor(*InnerPlatformFile, DirectoriesToSkip, DirectoriesToNotRecurse, false);
 
-			EnsureFileIsLocal(TestSyncFile);
+	/*TArray<FString> RootContentPaths;
+	FPackageName::QueryRootContentPaths(RootContentPaths); */
+	for (TArray<FString>::TConstIterator RootPathIt(ServerRootContentDirectories); RootPathIt; ++RootPathIt)
+	{
+		/*const FString& RootPath = *RootPathIt;
+		const FString& ContentFolder = FPackageName::LongPackageNameToFilename(RootPath);*/
+		const FString& ContentFolder = *RootPathIt;
+		InnerPlatformFile->IterateDirectory(*ContentFolder, Visitor);
+	}
 
-			if (!InnerPlatformFile->FileExists(*TestSyncFile) || InnerPlatformFile->FileSize(*TestSyncFile) < 1)
+	// delete out of date files using the server cached files
+	for (TMap<FString, FDateTime>::TIterator It(ServerCachedFiles); It; ++It)
+	{
+		bool bDeleteFile = bDeleteAllFiles;
+		FString ServerFile = It.Key();
+
+		// Convert the filename to the client version
+		ConvertServerFilenameToClientFilename(ServerFile);
+
+		// Set it in the visitor file times list
+		// If there is any pathing difference (relative path, or whatever) between the server's filelist and the results
+		// of platform directory iteration then this will Add a new entry rather than override the existing one.  This causes local file deletes
+		// and longer loads as we will never see the benefits of local device caching.
+		Visitor.FileTimes.Add(ServerFile, FDateTime::MinValue());
+
+		if (bDeleteFile == false)
+		{
+			// Check the time stamps...
+			// get local time
+			FDateTime LocalTime = InnerPlatformFile->GetTimeStamp(*ServerFile);
+			// If local time == MinValue than the file does not exist in the cache.
+			if (LocalTime != FDateTime::MinValue())
 			{
-				UE_LOG(LogNetworkPlatformFile, Fatal, TEXT("Could not sync test file %s."), *TestSyncFile);
+				FDateTime ServerTime = It.Value();
+				// delete if out of date
+				// We will use 1.0 second as the tolerance to cover any platform differences in resolution
+				FTimespan TimeDiff = LocalTime - ServerTime;
+				double TimeDiffInSeconds = TimeDiff.GetTotalSeconds();
+				bDeleteFile = (TimeDiffInSeconds > 1.0) || (TimeDiffInSeconds < -1.0);
+				if (bDeleteFile == true)
+				{
+					if (InnerPlatformFile->FileExists(*ServerFile) == true)
+					{
+						UE_LOG(LogNetworkPlatformFile, Display, TEXT("Deleting cached file: TimeDiff %5.3f, %s"), TimeDiffInSeconds, *It.Key());
+					}
+					else
+					{
+						// It's a directory
+						bDeleteFile = false;
+					}
+				}
+				else
+				{
+					UE_LOG(LogNetworkPlatformFile, Display, TEXT("Keeping cached file: %s, TimeDiff worked out ok"), *ServerFile);
+				}
 			}
 		}
+		if (bDeleteFile == true)
+		{
+			UE_LOG(LogNetworkPlatformFile, Display, TEXT("Deleting cached file: %s"), *ServerFile);
+			InnerPlatformFile->DeleteFile(*ServerFile);
+		}
 	}
 
-	FPlatformMisc::LowLevelOutputDebugStringf(TEXT("Network file startup time: %5.3f seconds\n"), NetworkFileStartupTime);
-
+	// Any content files we have locally that were not cached, delete them
+	for (TMap<FString, FDateTime>::TIterator It(Visitor.FileTimes); It; ++It)
+	{
+		if ( FCString::Stricmp( *FPaths::GetExtension( It.Key() ), TEXT("pak")) == 0 )
+		{
+			// ignore pak files they won't be mounted anyway 
+			continue;
+		}
+		if (It.Value() != FDateTime::MinValue())
+		{
+			// This was *not* found in the server file list... delete it
+			UE_LOG(LogNetworkPlatformFile, Display, TEXT("Deleting cached file: %s"), *It.Key());
+			InnerPlatformFile->DeleteFile(*It.Key());
+		}
+	}
 }
 
 FNetworkPlatformFile::~FNetworkPlatformFile()
@@ -772,7 +784,7 @@ void FNetworkPlatformFile::ConvertServerFilenameToClientFilename(FString& Filena
 	FNetworkPlatformFile::ConvertServerFilenameToClientFilename(FilenameToConvert, ServerEngineDir, ServerGameDir);
 }
 
-void FNetworkPlatformFile::FillGetFileList(FNetworkFileArchive& Payload, bool bInStreamingFileRequest)
+void FNetworkPlatformFile::FillGetFileList(FNetworkFileArchive& Payload)
 {
 	TArray<FString> TargetPlatformNames;
 	FPlatformMisc::GetValidTargetPlatforms(TargetPlatformNames);
@@ -798,10 +810,13 @@ void FNetworkPlatformFile::FillGetFileList(FNetworkFileArchive& Payload, bool bI
 	Payload << EngineRelPath;
 	Payload << GameRelPath;
 	Payload << Directories;
-	Payload << bInStreamingFileRequest;
+	Payload << ConnectionFlags;
+
+	FString VersionInfo = GetVersionInfo();
+	Payload << VersionInfo;
 }
 
-void FNetworkPlatformFile::ProcessServerInitialResponse(FArrayReader& InResponse, int32 OutServerPackageVersion, int32 OutServerPackageLicenseeVersion)
+void FNetworkPlatformFile::ProcessServerInitialResponse(FArrayReader& InResponse, int32& OutServerPackageVersion, int32& OutServerPackageLicenseeVersion)
 {
 	// Receive the cooked version information.
 	InResponse << OutServerPackageVersion;
@@ -827,6 +842,12 @@ void FNetworkPlatformFile::ProcessServerInitialResponse(FArrayReader& InResponse
 	}
 }
 
+
+FString FNetworkPlatformFile::GetVersionInfo() const
+{
+	return FString("");
+}
+
 bool FNetworkPlatformFile::SendReadMessage(uint8* Destination, int64 BytesToRead)
 {
 //	FScopeLock ScopeLock(&SynchronizationObject);
@@ -1348,21 +1369,16 @@ void FNetworkPlatformFile::PerformHeartbeat()
 		UE_LOG(LogNetworkPlatformFile, Log, TEXT("Server updated file '%s', deleting local copy %s"), *UpdatedFiles[FileIndex], *LocalFileName);
 
 		FString PackageName;
-		if ( FPackageName::TryConvertFilenameToLongPackageName(LocalFileName, PackageName) )
+		if (FPackageName::TryConvertFilenameToLongPackageName(LocalFileName, PackageName))
 		{
-			PackageNames.Add(PackageName );
+			PackageNames.Add(PackageName);
 		}
 		else
 		{
 			UE_LOG(LogNetworkPlatformFile, Log, TEXT("Unable to convert filename to package name %s"), *LocalFileName);
 		}
 
-		if (InnerPlatformFile->FileExists( *LocalFileName ) && InnerPlatformFile->DeleteFile(*LocalFileName) == false)
-		{
-			UE_LOG(LogNetworkPlatformFile, Error, TEXT("Failed to delete %s, someone is probably accessing without FNetworkPlatformFile, or we need better thread protection"), *UpdatedFiles[FileIndex]);
-		}
-		CachedLocalFiles.Remove(LocalFileName);
-		ServerFiles.AddFileOrDirectory(LocalFileName, FDateTime::UtcNow());
+		OnFileUpdated(LocalFileName);
 	}
 
 	if ( PackageNames.Num() > 0 )
@@ -1371,7 +1387,15 @@ void FNetworkPlatformFile::PerformHeartbeat()
 	}
 }
 
-
+void FNetworkPlatformFile::OnFileUpdated(const FString& LocalFileName)
+{
+	if (InnerPlatformFile->FileExists(*LocalFileName) && InnerPlatformFile->DeleteFile(*LocalFileName) == false)
+	{
+		UE_LOG(LogNetworkPlatformFile, Error, TEXT("Failed to delete %s, someone is probably accessing without FNetworkPlatformFile, or we need better thread protection"), *LocalFileName);
+	}
+	CachedLocalFiles.Remove(LocalFileName);
+	ServerFiles.AddFileOrDirectory(LocalFileName, FDateTime::UtcNow());
+}
 
 void FNetworkPlatformFile::ConvertServerFilenameToClientFilename(FString& FilenameToConvert, const FString& InServerEngineDir, const FString& InServerGameDir)
 {
@@ -1391,7 +1415,7 @@ void FNetworkPlatformFile::Tick()
 	static double StartTime = FPlatformTime::Seconds();
 
 	bool bShouldPerformHeartbeat = true;
-	if ((FPlatformTime::Seconds() - StartTime) > 5.0f )
+	if ((FPlatformTime::Seconds() - StartTime) > HeartbeatFrequency && HeartbeatFrequency >= 0 )
 	{
 
 
diff --git a/Engine/Source/Runtime/NetworkFile/Public/NetworkPlatformFile.h b/Engine/Source/Runtime/NetworkFile/Public/NetworkPlatformFile.h
--- a/Engine/Source/Runtime/NetworkFile/Public/NetworkPlatformFile.h
+++ b/Engine/Source/Runtime/NetworkFile/Public/NetworkPlatformFile.h
@@ -12,6 +12,7 @@ class FScopedEvent;
 
 DECLARE_LOG_CATEGORY_EXTERN(LogNetworkPlatformFile, Log, All);
 
+
 /**
  * Wrapper to redirect the low level file system to a server
  */
@@ -20,6 +21,7 @@ class NETWORKFILE_API FNetworkPlatformFile : public IPlatformFile, public FSelfR
 	friend class FAsyncFileSync;
 	friend void ReadUnsolicitedFile(int32 InNumUnsolictedFiles, FNetworkPlatformFile& InNetworkFile, IPlatformFile& InInnerPlatformFile,  FString& InServerEngineDir, FString& InServerGameDir);
 
+protected:
 	/**
 	 * Initialize network platform file give the specified host IP
 	 *
@@ -29,7 +31,10 @@ class NETWORKFILE_API FNetworkPlatformFile : public IPlatformFile, public FSelfR
 	 */
 	virtual bool InitializeInternal(IPlatformFile* Inner, const TCHAR* HostIP);
 
+	virtual void OnFileUpdated(const FString& LocalFilename);
+
 public:
+	
 
 	static const TCHAR* GetTypeName()
 	{
@@ -150,6 +155,10 @@ public:
 	static void ConvertServerFilenameToClientFilename(FString& FilenameToConvert, const FString& InServerEngineDir, const FString& InServerGameDir);
 
 
+	virtual FString GetVersionInfo() const;
+
+
+
 	//////////////////////////////////////////////////////////////////////////
 	// FSelfRegisteringExec interface
 	virtual bool Exec(class UWorld* InWorld, const TCHAR* Cmd, FOutputDevice& Ar ) override;
@@ -173,9 +182,10 @@ protected:
 	 */
 	virtual void ConvertServerFilenameToClientFilename(FString& FilenameToConvert);
 
-	virtual void FillGetFileList(FNetworkFileArchive& Payload, bool bInStreamingFileRequest);
+	virtual void FillGetFileList(FNetworkFileArchive& Payload);
 
-	virtual void ProcessServerInitialResponse(FArrayReader& InResponse, int32 OutServerPackageVersion, int32 OutServerPackageLicenseeVersion);
+	virtual void ProcessServerInitialResponse(FArrayReader& InResponse, int32& OutServerPackageVersion, int32& OutServerPackageLicenseeVersion);
+	virtual void ProcessServerCachedFilesResponse(FArrayReader& InReponse, const int32 ServerPackageVersion, const int32 ServerPackageLicenseeVersion );
 
 private:
 
@@ -248,6 +258,11 @@ protected:
 	bool				bIsUsable;
 	int32				FileServerPort;
 
+	// the connection flags are passed to the server during GetFileList
+	// the server may cache them
+	EConnectionFlags	ConnectionFlags;
+	// Frequency to send heartbeats to server in seconds set to negative number to disable
+	float HeartbeatFrequency;
 
 
 	// some stats for messuring network platform file performance
diff --git a/Engine/Source/Runtime/NetworkFileSystem/Private/NetworkFileServer.cpp b/Engine/Source/Runtime/NetworkFileSystem/Private/NetworkFileServer.cpp
--- a/Engine/Source/Runtime/NetworkFileSystem/Private/NetworkFileServer.cpp
+++ b/Engine/Source/Runtime/NetworkFileSystem/Private/NetworkFileServer.cpp
@@ -16,9 +16,8 @@ class FNetworkFileServerClientConnectionThreaded
 {
 public:
 
-	FNetworkFileServerClientConnectionThreaded(FSocket* InSocket, const FFileRequestDelegate& InFileRequestDelegate, 
-		const FRecompileShadersDelegate& InRecompileShadersDelegate, const FSandboxPathDelegate& SandboxPathOverrideDelegate, FOnFileModifiedDelegate* InOnFileModifiedCallback, const TArray<ITargetPlatform*>& InActiveTargetPlatforms )
-		:  FNetworkFileServerClientConnection( InFileRequestDelegate,InRecompileShadersDelegate, SandboxPathOverrideDelegate, InOnFileModifiedCallback, InActiveTargetPlatforms)
+	FNetworkFileServerClientConnectionThreaded(FSocket* InSocket, const FNetworkFileDelegateContainer* NetworkFileDelegates, const TArray<ITargetPlatform*>& InActiveTargetPlatforms )
+		:  FNetworkFileServerClientConnection( NetworkFileDelegates, InActiveTargetPlatforms)
 		  ,Socket(InSocket)
 	{
 		Running.Set(true);
@@ -102,7 +101,6 @@ public:
 	}
 
 private: 
-	FOnFileModifiedDelegate* OnFileModifiedCallback;
 	FSocket* Socket; 
 	FThreadSafeCounter StopRequested;
 	FThreadSafeCounter Running;
@@ -114,8 +112,7 @@ private:
 /* FNetworkFileServer constructors
  *****************************************************************************/
 
-FNetworkFileServer::FNetworkFileServer( int32 InPort, const FFileRequestDelegate* InFileRequestDelegate,const FRecompileShadersDelegate* InRecompileShadersDelegate, 
-	const FSandboxPathDelegate* InSandboxPathOverrideDelegate, FOnFileModifiedDelegate* InOnFileModifiedCallback,
+FNetworkFileServer::FNetworkFileServer( int32 InPort, FNetworkFileDelegateContainer InNetworkFileDelegateContainer,
 	const TArray<ITargetPlatform*>& InActiveTargetPlatforms )
 	:ActiveTargetPlatforms(InActiveTargetPlatforms)
 {
@@ -128,22 +125,7 @@ FNetworkFileServer::FNetworkFileServer( int32 InPort, const FFileRequestDelegate
 	StopRequested.Set(false);
 	UE_LOG(LogFileServer, Warning, TEXT("Unreal Network File Server starting up..."));
 
-	if (InFileRequestDelegate && InFileRequestDelegate->IsBound())
-	{
-		FileRequestDelegate = *InFileRequestDelegate;
-	}
-
-	if (InRecompileShadersDelegate && InRecompileShadersDelegate->IsBound())
-	{
-		RecompileShadersDelegate = *InRecompileShadersDelegate;
-	}
-
-	if ( InSandboxPathOverrideDelegate && InSandboxPathOverrideDelegate->IsBound() )
-	{
-		SandboxPathOverrideDelegate = *InSandboxPathOverrideDelegate;
-	}
-
-	OnFileModifiedCallback = InOnFileModifiedCallback;
+	NetworkFileDelegates = InNetworkFileDelegateContainer;
 
 	// make sure sockets are going
 	ISocketSubsystem* SocketSubsystem = ISocketSubsystem::Get();
@@ -265,7 +247,7 @@ uint32 FNetworkFileServer::Run( )
 					}
 				}
 
-				FNetworkFileServerClientConnectionThreaded* Connection = new FNetworkFileServerClientConnectionThreaded(ClientSocket, FileRequestDelegate, RecompileShadersDelegate, SandboxPathOverrideDelegate, OnFileModifiedCallback, ActiveTargetPlatforms);
+				FNetworkFileServerClientConnectionThreaded* Connection = new FNetworkFileServerClientConnectionThreaded(ClientSocket, &NetworkFileDelegates, ActiveTargetPlatforms);
 				Connections.Add(Connection);
 				UE_LOG(LogFileServer, Display, TEXT( "Client %s connected." ), *Connection->GetDescription() );
 			}
diff --git a/Engine/Source/Runtime/NetworkFileSystem/Private/NetworkFileServer.h b/Engine/Source/Runtime/NetworkFileSystem/Private/NetworkFileServer.h
--- a/Engine/Source/Runtime/NetworkFileSystem/Private/NetworkFileServer.h
+++ b/Engine/Source/Runtime/NetworkFileSystem/Private/NetworkFileServer.h
@@ -27,8 +27,7 @@ public:
 	 * @param InPort The port number to bind to (0 = any available port).
 	 * @param InFileRequestDelegate 
 	 */
-	FNetworkFileServer( int32 InPort, const FFileRequestDelegate* InFileRequestDelegate, 
-		const FRecompileShadersDelegate* InRecompileShadersDelegate, const FSandboxPathDelegate* SandboxPathOverrideDelegate, FOnFileModifiedDelegate* OnFileModifiedCallback, const TArray<ITargetPlatform*>& InActiveTargetPlatforms );
+	FNetworkFileServer( int32 InPort, FNetworkFileDelegateContainer InNetworkFileDelegateContainer, const TArray<ITargetPlatform*>& InActiveTargetPlatforms );
 
 	/**
 	 * Destructor.
@@ -81,17 +80,7 @@ private:
 
 public:
 
-	// Holds a delegate to be invoked on every sync request.
-	FFileRequestDelegate FileRequestDelegate;
-
-	// Holds a delegate to be invoked when a client requests a shader recompile.
-	FRecompileShadersDelegate RecompileShadersDelegate;
-
-	FSandboxPathDelegate SandboxPathDelegate;
-
-	FSandboxPathDelegate SandboxPathOverrideDelegate;
-
-	FOnFileModifiedDelegate* OnFileModifiedCallback;
+	FNetworkFileDelegateContainer NetworkFileDelegates;
 
 	// cached copy of the active target platforms (if any)
 	const TArray<ITargetPlatform*> ActiveTargetPlatforms;
