Backout 5372217. This fixes COTF on some platforms.

diff --git a/Engine/Source/Runtime/NetworkFile/Private/NetworkPlatformFile.cpp b/Engine/Source/Runtime/NetworkFile/Private/NetworkPlatformFile.cpp
--- a/Engine/Source/Runtime/NetworkFile/Private/NetworkPlatformFile.cpp
+++ b/Engine/Source/Runtime/NetworkFile/Private/NetworkPlatformFile.cpp
@@ -478,26 +478,6 @@ IFileHandle* FNetworkPlatformFile::OpenRead(const TCHAR* Filename, bool bAllowWr
 	return Result;
 }
 
-IAsyncReadFileHandle* FNetworkPlatformFile::OpenAsyncRead(const TCHAR* Filename)
-{
-	FString RelativeFilename = Filename;
-	MakeStandardNetworkFilename(RelativeFilename);
-	// don't copy files in local directories
-	if (!IsInLocalDirectory(RelativeFilename))
-	{
-		EnsureFileIsLocal(RelativeFilename);
-	}
-
-	double StartTime;
-	float ThisTime;
-
-	StartTime = FPlatformTime::Seconds();
-	IAsyncReadFileHandle* Result = InnerPlatformFile->OpenAsyncRead(Filename);
-
-	ThisTime = 1000.0f * float(FPlatformTime::Seconds() - StartTime);
-	return Result;
-}
-
 IFileHandle* FNetworkPlatformFile::OpenWrite(const TCHAR* Filename, bool bAppend, bool bAllowRead)
 {
 //	FScopeLock ScopeLock(&SynchronizationObject);
diff --git a/Engine/Source/Runtime/NetworkFile/Public/NetworkPlatformFile.h b/Engine/Source/Runtime/NetworkFile/Public/NetworkPlatformFile.h
--- a/Engine/Source/Runtime/NetworkFile/Public/NetworkPlatformFile.h
+++ b/Engine/Source/Runtime/NetworkFile/Public/NetworkPlatformFile.h
@@ -119,7 +119,6 @@ public:
 		return Filename;
 	}
 	virtual IFileHandle*	OpenRead(const TCHAR* Filename, bool bAllowWrite = false) override;
-	virtual IAsyncReadFileHandle* OpenAsyncRead(const TCHAR* Filename) override;
 	virtual IFileHandle*	OpenWrite(const TCHAR* Filename, bool bAppend = false, bool bAllowRead = false) override;
 	virtual bool		DirectoryExists(const TCHAR* Directory) override;
 	virtual bool		CreateDirectoryTree(const TCHAR* Directory) override;
