Cook on the fly changes:
 - target platform code can manage the file transfer from the cook server.
 - client ip address and platform-specific metadata is sent from the client.

#rnx

diff --git a/Engine/Source/Developer/TargetPlatform/Public/Common/TargetPlatformBase.h b/Engine/Source/Developer/TargetPlatform/Public/Common/TargetPlatformBase.h
--- a/Engine/Source/Developer/TargetPlatform/Public/Common/TargetPlatformBase.h
+++ b/Engine/Source/Developer/TargetPlatform/Public/Common/TargetPlatformBase.h
@@ -163,6 +163,11 @@ public:
 	}
 #endif
 
+	virtual bool CopyFileToTarget(const FString& TargetAddress, const FString& HostFilename, const FString& TargetFilename, const TMap<FString,FString>& CustomPlatformData) override
+	{
+		return false; 
+	}
+
 protected:
 
 	FTargetPlatformBase(const PlatformInfo::FPlatformInfo *const InPlatformInfo)
diff --git a/Engine/Source/Developer/TargetPlatform/Public/Interfaces/ITargetPlatform.h b/Engine/Source/Developer/TargetPlatform/Public/Interfaces/ITargetPlatform.h
--- a/Engine/Source/Developer/TargetPlatform/Public/Interfaces/ITargetPlatform.h
+++ b/Engine/Source/Developer/TargetPlatform/Public/Interfaces/ITargetPlatform.h
@@ -602,6 +602,12 @@ public:
 	 * Returns wheter or not this 16bit index buffer should be promoted to 32bit
 	 */
 	virtual bool ShouldExpandTo32Bit(const uint16* Indices, const int32 NumIndices) const = 0;
+
+	/**
+	 * Copy a file to the target
+	 */
+	virtual bool CopyFileToTarget(const FString& DeviceId, const FString& HostFilename, const FString& TargetFilename, const TMap<FString,FString>& CustomPlatformData) = 0;
+
 public:
 
 	/**
diff --git a/Engine/Source/Runtime/Core/Public/GenericPlatform/GenericPlatformFile.h b/Engine/Source/Runtime/Core/Public/GenericPlatform/GenericPlatformFile.h
--- a/Engine/Source/Runtime/Core/Public/GenericPlatform/GenericPlatformFile.h
+++ b/Engine/Source/Runtime/Core/Public/GenericPlatform/GenericPlatformFile.h
@@ -292,6 +292,19 @@ public:
 	/** For case insensitive filesystems, returns the full path of the file with the same case as in the filesystem */
 	virtual FString GetFilenameOnDisk(const TCHAR* Filename) = 0;
 
+	/** Return the size of the file, or -1 if it doesn't exist. **/
+	virtual int64		FileSize(const TCHAR* Filename, bool bUseWriteLocation)
+	{
+		return FileSize(Filename);
+	}
+
+	/** Check if a file exists in a location suitable for writing at runtime */
+	virtual bool		FileExists(const TCHAR* Filename, bool bUseWriteLocation)
+	{
+		return FileSize(Filename, bUseWriteLocation) >= 0;
+	}
+
+
 	/** Attempt to open a file for reading.
 	 *
 	 * @param Filename file to be opened
diff --git a/Engine/Source/Runtime/Core/Public/GenericPlatform/GenericPlatformMisc.h b/Engine/Source/Runtime/Core/Public/GenericPlatform/GenericPlatformMisc.h
--- a/Engine/Source/Runtime/Core/Public/GenericPlatform/GenericPlatformMisc.h
+++ b/Engine/Source/Runtime/Core/Public/GenericPlatform/GenericPlatformMisc.h
@@ -1436,6 +1436,13 @@ public:
 		return false;
 	}
 
+	/**
+	 * Returns any platform-specific key-value data that needs to be sent to the network file server
+	 */
+	FORCEINLINE static void GetNetworkFileCustomData(TMap<FString,FString>& OutCustomPlatformData)
+	{
+	}
+
 #if !UE_BUILD_SHIPPING
 	/**
 	 * Returns any platform specific warning messages we want printed on screen
diff --git a/Engine/Source/Runtime/Core/Public/HAL/IPlatformFileCachedWrapper.h b/Engine/Source/Runtime/Core/Public/HAL/IPlatformFileCachedWrapper.h
--- a/Engine/Source/Runtime/Core/Public/HAL/IPlatformFileCachedWrapper.h
+++ b/Engine/Source/Runtime/Core/Public/HAL/IPlatformFileCachedWrapper.h
@@ -313,10 +313,18 @@ public:
 	{
 		return LowerLevel->FileExists(Filename);
 	}
+	virtual bool		FileExists(const TCHAR* Filename, bool bUseWriteLocation) override
+	{
+		return LowerLevel->FileExists(Filename, bUseWriteLocation);
+	}
 	virtual int64		FileSize(const TCHAR* Filename) override
 	{
 		return LowerLevel->FileSize(Filename);
 	}
+	virtual int64		FileSize(const TCHAR* Filename, bool bUseWriteLocation) override
+	{
+		return LowerLevel->FileSize(Filename, bUseWriteLocation);
+	}
 	virtual bool		DeleteFile(const TCHAR* Filename) override
 	{
 		return LowerLevel->DeleteFile(Filename);
diff --git a/Engine/Source/Runtime/NetworkFile/Private/NetworkPlatformFile.cpp b/Engine/Source/Runtime/NetworkFile/Private/NetworkPlatformFile.cpp
--- a/Engine/Source/Runtime/NetworkFile/Private/NetworkPlatformFile.cpp
+++ b/Engine/Source/Runtime/NetworkFile/Private/NetworkPlatformFile.cpp
@@ -15,6 +15,7 @@
 #include "Modules/ModuleManager.h"
 #include "DerivedDataCacheInterface.h"
 #include "Misc/PackageName.h"
+#include "SocketSubsystem.h"
 
 #if ENABLE_HTTP_FOR_NETWORK_FILE
 #include "HTTPTransport.h"
@@ -221,7 +222,7 @@ void FNetworkPlatformFile::InitializeAfterSetActive()
 			InnerPlatformFile->DeleteFile(*TestSyncFile);
 			if (InnerPlatformFile->FileExists(*TestSyncFile))
 			{
-				UE_LOG(LogNetworkPlatformFile, Fatal, TEXT("Could not delete file sync test file %s."), *TestSyncFile);
+				UE_LOG(LogNetworkPlatformFile, Error, TEXT("Could not delete file sync test file %s."), *TestSyncFile);
 			}
 
 			EnsureFileIsLocal(TestSyncFile);
@@ -251,9 +252,9 @@ void FNetworkPlatformFile::ProcessServerCachedFilesResponse(FArrayReader& Respon
 	// Check the stored cooked version
 	FString CookedVersionFile = FPaths::GeneratedConfigDir() / TEXT("CookedVersion.txt");
 
-	if (InnerPlatformFile->FileExists(*CookedVersionFile) == true)
+	if (InnerPlatformFile->FileExists(*CookedVersionFile, true) == true)
 	{
-		IFileHandle* FileHandle = InnerPlatformFile->OpenRead(*CookedVersionFile);
+		IFileHandle* FileHandle = InnerPlatformFile->OpenRead(*CookedVersionFile, true);
 		if (FileHandle != NULL)
 		{
 			int32 StoredPackageCookedVersion;
@@ -815,6 +816,15 @@ void FNetworkPlatformFile::FillGetFileList(FNetworkFileArchive& Payload)
 
 	FString VersionInfo = GetVersionInfo();
 	Payload << VersionInfo;
+
+	bool bBindAllDummy;
+	TSharedRef<FInternetAddr> HostAddr = ISocketSubsystem::Get(PLATFORM_SOCKETSUBSYSTEM)->GetLocalHostAddr(*GLog, bBindAllDummy);
+	FString HostAddress = HostAddr->ToString(false);
+	Payload << HostAddress;
+
+	TMap<FString,FString> CustomPlatformData;
+	FPlatformMisc::GetNetworkFileCustomData(CustomPlatformData);
+	Payload << CustomPlatformData;
 }
 
 void FNetworkPlatformFile::ProcessServerInitialResponse(FArrayReader& InResponse, int32& OutServerPackageVersion, int32& OutServerPackageLicenseeVersion)
@@ -955,58 +965,63 @@ public:
 		*FileArchive << FileSize;
 		if (ServerTimeStamp != FDateTime::MinValue())  // if the file didn't actually exist on the server, don't create a zero byte file
 		{
-			FString TempFilename = Filename + TEXT(".tmp");
-			InnerPlatformFile.CreateDirectoryTree(*FPaths::GetPath(Filename));
+			// -1 filesize means that we already copied it via TargetPlatform
+			if (FileSize != MAX_uint64)
 			{
-				TUniquePtr<IFileHandle> FileHandle;
-				FileHandle.Reset(InnerPlatformFile.OpenWrite(*TempFilename));
-
-				if (!FileHandle)
+				FString TempFilename = Filename + TEXT(".tmp");
+				InnerPlatformFile.CreateDirectoryTree(*FPaths::GetPath(Filename));
 				{
-					UE_LOG(LogNetworkPlatformFile, Fatal, TEXT("Could not open file for writing '%s'."), *TempFilename);
-				}
+					TUniquePtr<IFileHandle> FileHandle;
+					FileHandle.Reset(InnerPlatformFile.OpenWrite(*TempFilename));
 
-				// now write the file from bytes pulled from the archive
-				// read/write a chunk at a time
-				uint64 RemainingData = FileSize;
-				while (RemainingData)
-				{
-					// read next chunk from archive
-					uint32 LocalSize = FPlatformMath::Min<uint32>(UE_ARRAY_COUNT(Buffer), RemainingData);
-					FileArchive->Serialize(Buffer, LocalSize);
-					// write it out
-					if (!FileHandle->Write(Buffer, LocalSize))
+					if (!FileHandle)
 					{
-						UE_LOG(LogNetworkPlatformFile, Fatal, TEXT("Could not write '%s'."), *TempFilename);
+						UE_LOG(LogNetworkPlatformFile, Fatal, TEXT("Could not open file for writing '%s'."), *TempFilename);
 					}
 
-					// decrement how much is left
-					RemainingData -= LocalSize;
-				}
+					// now write the file from bytes pulled from the archive
+					// read/write a chunk at a time
+					uint64 RemainingData = FileSize;
+					while (RemainingData)
+					{
+						// read next chunk from archive
+						uint32 LocalSize = FPlatformMath::Min<uint32>(UE_ARRAY_COUNT(Buffer), RemainingData);
+						FileArchive->Serialize(Buffer, LocalSize);
+						// write it out
+						if (!FileHandle->Write(Buffer, LocalSize))
+						{
+							UE_LOG(LogNetworkPlatformFile, Fatal, TEXT("Could not write '%s'."), *TempFilename);
+						}
+
+						// decrement how much is left
+						RemainingData -= LocalSize;
+					}
 
-				// delete async write archives
-				if (Event)
-				{
-					delete FileArchive;
-				}
+					// delete async write archives
+					if (Event)
+					{
+						delete FileArchive;
+					}
 
-				if (InnerPlatformFile.FileSize(*TempFilename) != FileSize)
-				{
-					UE_LOG(LogNetworkPlatformFile, Fatal, TEXT("Did not write '%s'."), *TempFilename);
+					if (InnerPlatformFile.FileSize(*TempFilename) != FileSize)
+					{
+						UE_LOG(LogNetworkPlatformFile, Fatal, TEXT("Did not write '%s'."), *TempFilename);
+					}
 				}
-			}
 
-			// rename from temp filename to real filename
-			InnerPlatformFile.MoveFile(*Filename, *TempFilename);
+				// rename from temp filename to real filename
+				InnerPlatformFile.MoveFile(*Filename, *TempFilename);
 
-			// now set the server's timestamp on the local file (so we can make valid comparisons)
-			InnerPlatformFile.SetTimeStamp(*Filename, ServerTimeStamp);
+				// now set the server's timestamp on the local file (so we can make valid comparisons)
+				InnerPlatformFile.SetTimeStamp(*Filename, ServerTimeStamp);
 
-			FDateTime CheckTime = InnerPlatformFile.GetTimeStamp(*Filename);
-			if (CheckTime < ServerTimeStamp)
-			{
-				UE_LOG(LogNetworkPlatformFile, Fatal, TEXT("Could Not Set Timestamp '%s'  %s < %s."), *Filename, *CheckTime.ToString(), *ServerTimeStamp.ToString());
+			    FDateTime CheckTime = InnerPlatformFile.GetTimeStamp(*Filename);
+			    if (CheckTime < ServerTimeStamp)
+			    {
+				    UE_LOG(LogNetworkPlatformFile, Fatal, TEXT("Could Not Set Timestamp '%s'  %s < %s."), *Filename, *CheckTime.ToString(), *ServerTimeStamp.ToString());
+			    }
 			}
+
 		}
 		if (Event)
 		{
diff --git a/Engine/Source/Runtime/NetworkFileSystem/Private/NetworkFileServerConnection.cpp b/Engine/Source/Runtime/NetworkFileSystem/Private/NetworkFileServerConnection.cpp
--- a/Engine/Source/Runtime/NetworkFileSystem/Private/NetworkFileServerConnection.cpp
+++ b/Engine/Source/Runtime/NetworkFileSystem/Private/NetworkFileServerConnection.cpp
@@ -163,7 +163,7 @@ FString FNetworkFileServerClientConnection::FixupSandboxPathForClient(const FStr
 	return Fixed;
 }
 
-/*void FNetworkFileServerClientConnection::ConvertServerFilenameToClientFilename(FString& FilenameToConvert)
+static void ConvertServerFilenameToClientFilename(FString& FilenameToConvert, const FString& ConnectedEngineDir, const FString& ConnectedProjectDir)
 {
 
 	if (FilenameToConvert.StartsWith(FPaths::EngineDir()))
@@ -185,7 +185,7 @@ FString FNetworkFileServerClientConnection::FixupSandboxPathForClient(const FStr
 			FilenameToConvert = FilenameToConvert.Replace(*(FPaths::ProjectDir()), *ConnectedProjectDir);
 	}
 #endif
-}*/
+}
 
 static FCriticalSection SocketCriticalSection;
 
@@ -352,7 +352,9 @@ bool FNetworkFileServerClientConnection::ProcessPayload(FArchive& Ar)
 			for (int32 Index = 0; Index < NumUnsolictedFiles; Index++)
 			{
 				FBufferArchive OutUnsolicitedFile;
-				PackageFile(UnsolictedFiles[Index], OutUnsolicitedFile);
+				FString TargetFilename = UnsolictedFiles[Index];
+				ConvertServerFilenameToClientFilename(TargetFilename, ConnectedEngineDir, ConnectedProjectDir);
+				PackageFile(UnsolictedFiles[Index], TargetFilename, OutUnsolicitedFile);
 
 				UE_LOG(LogFileServer, Display, TEXT("Returning unsolicited file %s with %d bytes"), *UnsolictedFiles[Index], OutUnsolicitedFile.Num());
 
@@ -735,10 +737,13 @@ bool FNetworkFileServerClientConnection::ProcessGetFileList( FArchive& In, FArch
 	FString EnginePlatformExtensionsRelativePath;
 	FString ProjectPlatformExtensionsRelativePath;
 	TArray<FString> RootDirectories;
+	TMap<FString,FString> CustomPlatformData;
+
 	
 	EConnectionFlags ConnectionFlags;
 
 	FString ClientVersionInfo;
+	FString TargetAddress;
 
 	In << TargetPlatformNames;
 	In << GameName;
@@ -749,6 +754,8 @@ bool FNetworkFileServerClientConnection::ProcessGetFileList( FArchive& In, FArch
 	In << RootDirectories;
 	In << ConnectionFlags;
 	In << ClientVersionInfo;
+	In << TargetAddress;
+	In << CustomPlatformData;
 
 	if ( NetworkFileDelegates->NewConnectionDelegate.IsBound() )
 	{
@@ -767,6 +774,10 @@ bool FNetworkFileServerClientConnection::ProcessGetFileList( FArchive& In, FArch
 	const bool bIsPrecookedIterativeRequest = (ConnectionFlags & EConnectionFlags::PreCookedIterative) == EConnectionFlags::PreCookedIterative;
 
 	ConnectedPlatformName = TEXT("");
+	ConnectedTargetPlatform = nullptr;
+	ConnectedIPAddress = TEXT("");
+	ConnectedTargetCustomData.Reset();
+
 
 	// if we didn't find one (and this is a dumb server - no active platforms), then just use what was sent
 	if (ActiveTargetPlatforms.Num() == 0)
@@ -789,6 +800,9 @@ bool FNetworkFileServerClientConnection::ProcessGetFileList( FArchive& In, FArch
 				{
 					bSendLowerCase = ActiveTargetPlatforms[ActiveTPIndex]->SendLowerCaseFilePaths();
 					ConnectedPlatformName = ActiveTargetPlatforms[ActiveTPIndex]->PlatformName();
+					ConnectedTargetPlatform = ActiveTargetPlatforms[ActiveTPIndex];
+					ConnectedIPAddress = TargetAddress;
+					ConnectedTargetCustomData = MoveTemp(CustomPlatformData);
 					break;
 				}
 			}
@@ -1093,11 +1107,23 @@ void FNetworkFileServerClientConnection::ProcessHeartbeat( FArchive& In, FArchiv
 /* FStreamingNetworkFileServerConnection callbacks
  *****************************************************************************/
 
-bool FNetworkFileServerClientConnection::PackageFile( FString& Filename, FArchive& Out )
+bool FNetworkFileServerClientConnection::PackageFile( FString& Filename, FString& TargetFilename, FArchive& Out )
 {
 	// get file timestamp and send it to client
 	FDateTime ServerTimeStamp = Sandbox->GetTimeStamp(*Filename);
 
+	FString AbsHostFile = Sandbox->ConvertToAbsolutePathForExternalAppForRead(*Filename);
+	if (ConnectedTargetPlatform != nullptr && ConnectedTargetPlatform->CopyFileToTarget(ConnectedIPAddress, AbsHostFile, TargetFilename, ConnectedTargetCustomData))
+	{
+		Out << Filename;
+		Out << ServerTimeStamp;
+		// MAX_uint64 here indicates that it was copied already
+		uint64 FileSize = MAX_uint64;
+		Out << FileSize;
+
+		return true;
+	}
+
 	TArray<uint8> Contents;
 	// open file
 	IFileHandle* File = Sandbox->OpenRead(*Filename);
@@ -1175,6 +1201,7 @@ void FNetworkFileServerClientConnection::ProcessSyncFile( FArchive& In, FArchive
 	
 	UE_LOG(LogFileServer, Verbose, TEXT("Try sync file %s"), *Filename);
 
+	FString ClientFilename = Filename;
 	ConvertClientFilenameToServerFilename(Filename);
 	
 	//FString AbsFile(FString(*Sandbox->ConvertToAbsolutePathForExternalApp(*Filename)).MakeStandardFilename());
@@ -1196,7 +1223,7 @@ void FNetworkFileServerClientConnection::ProcessSyncFile( FArchive& In, FArchive
 		}
 	}
 
-	PackageFile(Filename, Out);
+	PackageFile(Filename, ClientFilename, Out);
 
 	PackageFileTime += 1000.0f * float(FPlatformTime::Seconds() - StartTime);
 }
diff --git a/Engine/Source/Runtime/NetworkFileSystem/Private/NetworkFileServerConnection.h b/Engine/Source/Runtime/NetworkFileSystem/Private/NetworkFileServerConnection.h
--- a/Engine/Source/Runtime/NetworkFileSystem/Private/NetworkFileServerConnection.h
+++ b/Engine/Source/Runtime/NetworkFileSystem/Private/NetworkFileServerConnection.h
@@ -135,7 +135,7 @@ protected:
 		return OpenFile ? *OpenFile : NULL;
 	}
 
-	bool PackageFile( FString& Filename, FArchive& Out);
+	bool PackageFile( FString& Filename, FString& TargetFilename, FArchive& Out);
 
 	/**
 	 * Processes a RecompileShaders message.
@@ -191,6 +191,9 @@ private:
 	// Hold the name of the currently connected platform.
 	FString ConnectedPlatformName;
 
+	// Hold the ip address of the currently connected platform.
+	FString ConnectedIPAddress;
+
 	// Hold the engine directory from the connected platform.
 	FString ConnectedEngineDir;
 
@@ -256,6 +259,12 @@ private:
 	// cached copy of the active target platforms (if any)
 	const TArray<ITargetPlatform*>& ActiveTargetPlatforms;
 
+	// connected TargetPlatform object, if ActiveTargetPlatforms existed
+	ITargetPlatform* ConnectedTargetPlatform;
+
+	// custom key-value pair data for the curently connected target platform
+	TMap<FString,FString> ConnectedTargetCustomData;
+
 
 	//////////////////////////////////////////////////////////////////////////
 	//stats
diff --git a/Engine/Source/Runtime/SandboxFile/Private/IPlatformFileSandboxWrapper.cpp b/Engine/Source/Runtime/SandboxFile/Private/IPlatformFileSandboxWrapper.cpp
--- a/Engine/Source/Runtime/SandboxFile/Private/IPlatformFileSandboxWrapper.cpp
+++ b/Engine/Source/Runtime/SandboxFile/Private/IPlatformFileSandboxWrapper.cpp
@@ -359,7 +359,7 @@ const FString& FSandboxPlatformFile::GetAbsoluteGameDirectory()
 {
 	if (AbsoluteGameDirectory.IsEmpty())
 	{
-		AbsoluteGameDirectory = FPaths::GetProjectFilePath();
+		AbsoluteGameDirectory = FPaths::ProjectDir();
 		UE_CLOG(AbsoluteGameDirectory.IsEmpty(), SandboxFile, Fatal, TEXT("SandboxFileWrapper tried to access project path before it was set."));
 		AbsoluteGameDirectory = FPaths::ConvertRelativePathToFull(AbsoluteGameDirectory);
 		// Strip .uproject filename
