Fix static analysis warning

diff --git a/Engine/Source/Runtime/NetworkFile/Private/NetworkPlatformFile.cpp b/Engine/Source/Runtime/NetworkFile/Private/NetworkPlatformFile.cpp
--- a/Engine/Source/Runtime/NetworkFile/Private/NetworkPlatformFile.cpp
+++ b/Engine/Source/Runtime/NetworkFile/Private/NetworkPlatformFile.cpp
@@ -953,7 +953,7 @@ public:
 	{
 		// Read FileSize first so that the correct amount of data is read from the archive
 		// before exiting this worker.
-		uint64 FileSize;
+		uint64 FileSize = 0;
 		*FileArchive << FileSize;
 
 		bool bCopiedExternally = (FileSize == MAX_uint64); // -1 filesize means that we already copied it via TargetPlatform
