NetworkPlatformFile: Improve initial sync of server files
- updated local file cache pruning to not to delete files that doesn't exist on client
- added option to timestamp directory visitor to lower case filenames and directory names
- added -DeleteLocalCache option to force delete all local file(s)

#rnx

diff --git a/Engine/Source/Runtime/Core/Private/Misc/LocalTimestampDirectoryVisitor.cpp b/Engine/Source/Runtime/Core/Private/Misc/LocalTimestampDirectoryVisitor.cpp
--- a/Engine/Source/Runtime/Core/Private/Misc/LocalTimestampDirectoryVisitor.cpp
+++ b/Engine/Source/Runtime/Core/Private/Misc/LocalTimestampDirectoryVisitor.cpp
@@ -6,8 +6,9 @@
 /* FLocalTimestampVisitor structors
  *****************************************************************************/
 
-FLocalTimestampDirectoryVisitor::FLocalTimestampDirectoryVisitor( IPlatformFile& InFileInterface, const TArray<FString>& InDirectoriesToIgnore, const TArray<FString>& InDirectoriesToNotRecurse, bool bInCacheDirectories )
+FLocalTimestampDirectoryVisitor::FLocalTimestampDirectoryVisitor( IPlatformFile& InFileInterface, const TArray<FString>& InDirectoriesToIgnore, const TArray<FString>& InDirectoriesToNotRecurse, bool bInCacheDirectories, bool bInMakeLowerCase )
 	: bCacheDirectories(bInCacheDirectories)
+	, bMakeLowerCase(bInMakeLowerCase)
 	, FileInterface(InFileInterface)
 {
 	// make sure the paths are standardized, since the Visitor will assume they are standard
@@ -36,20 +37,19 @@ bool FLocalTimestampDirectoryVisitor::Visit(const TCHAR* FilenameOrDirectory, bo
 	FString RelativeFilename = FilenameOrDirectory;
 	FPaths::MakeStandardFilename(RelativeFilename);
 
+	if (bMakeLowerCase)
+	{
+		RelativeFilename.ToLowerInline();
+	}
+
 	// cache files and optionally directories
 	if (!bIsDirectory)
 	{
-		FileTimes.Add(RelativeFilename, FileInterface.GetTimeStamp(FilenameOrDirectory));
+		FileTimes.Emplace(MoveTemp(RelativeFilename), FileInterface.GetTimeStamp(FilenameOrDirectory));
 	}
-	else if (bCacheDirectories)
-	{
-		// we use a timestamp of 0 to indicate a directory
-		FileTimes.Add(RelativeFilename, 0);
-	}
-
-	// iterate over directories we care about
-	if (bIsDirectory)
+	else 
 	{
+		// iterate over directories we care about
 		bool bShouldRecurse = true;
 		// look in all the ignore directories looking for a match
 		for (int32 DirIndex = 0; DirIndex < DirectoriesToIgnore.Num() && bShouldRecurse; DirIndex++)
@@ -78,6 +78,12 @@ bool FLocalTimestampDirectoryVisitor::Visit(const TCHAR* FilenameOrDirectory, bo
 			}
 		}
 
+		if (bCacheDirectories)
+		{
+			// we use a timestamp of 0 to indicate a directory
+			FileTimes.Emplace(MoveTemp(RelativeFilename), 0);
+		}
+
 		// recurse if we should
 		if (bShouldRecurse)
 		{
diff --git a/Engine/Source/Runtime/Core/Public/Misc/LocalTimestampDirectoryVisitor.h b/Engine/Source/Runtime/Core/Public/Misc/LocalTimestampDirectoryVisitor.h
--- a/Engine/Source/Runtime/Core/Public/Misc/LocalTimestampDirectoryVisitor.h
+++ b/Engine/Source/Runtime/Core/Public/Misc/LocalTimestampDirectoryVisitor.h
@@ -27,8 +27,9 @@ public:
 	 * @param InDirectoriesToIgnore - The list of directories to ignore.
 	 * @param InDirectoriesToNotRecurse - The list of directories to not visit recursively.
 	 * @param bInCacheDirectories - Whether to cache the directories.
+	 * @param bInMakeLowerCase - Whether to lower case filenames and directories.
 	 */
-	FLocalTimestampDirectoryVisitor( IPlatformFile& InFileInterface, const TArray<FString>& InDirectoriesToIgnore, const TArray<FString>& InDirectoriesToNotRecurse, bool bInCacheDirectories = false );
+	FLocalTimestampDirectoryVisitor( IPlatformFile& InFileInterface, const TArray<FString>& InDirectoriesToIgnore, const TArray<FString>& InDirectoriesToNotRecurse, bool bInCacheDirectories = false, bool bInMakeLowerCase = false );
 
 public:
 
@@ -41,6 +42,9 @@ private:
 	// true if we want directories in this list. */
 	bool bCacheDirectories;
 
+	// true if all filenames and directories should be lower cased. */
+	bool bMakeLowerCase;
+
 	// Holds a list of directories that we should not traverse into. */
 	TArray<FString> DirectoriesToIgnore;
 
diff --git a/Engine/Source/Runtime/NetworkFile/Private/NetworkPlatformFile.cpp b/Engine/Source/Runtime/NetworkFile/Private/NetworkPlatformFile.cpp
--- a/Engine/Source/Runtime/NetworkFile/Private/NetworkPlatformFile.cpp
+++ b/Engine/Source/Runtime/NetworkFile/Private/NetworkPlatformFile.cpp
@@ -46,9 +46,6 @@ FNetworkPlatformFile::FNetworkPlatformFile()
 	, FinishedAsyncWriteUnsolicitedFiles(NULL)
 	, Transport(NULL)
 {
-
-
-	
 	TotalWriteTime = 0.0; // total non async time spent writing to disk
 	TotalNetworkSyncTime = 0.0; // total non async time spent syncing to network
 	TotalTimeSpentInUnsolicitedPackages = 0.0; // total time async processing unsolicited packages
@@ -58,8 +55,6 @@ FNetworkPlatformFile::FNetworkPlatformFile()
 	TotalUnsolicitedPackages = 0; // total number unsolicited files synced  
 	UnsolicitedPackagesHits = 0; // total number of hits from waiting on unsolicited packages
 	UnsolicitedPackageWaits = 0; // total number of waits on unsolicited packages
-	
-
 }
 
 bool FNetworkPlatformFile::ShouldBeUsed(IPlatformFile* Inner, const TCHAR* CmdLine) const
@@ -247,7 +242,21 @@ void FNetworkPlatformFile::ProcessServerCachedFilesResponse(FArrayReader& Respon
 
 	// receive a list of the cache files and their timestamps
 	TMap<FString, FDateTime> ServerCachedFiles;
-	Response << ServerCachedFiles;
+	{
+		TMap<FString, FDateTime> ServerCachedFileResponse;
+		Response << ServerCachedFileResponse;
+
+		ServerCachedFiles.Reserve(ServerCachedFileResponse.Num());
+		for (const auto& ServerFileTimePair : ServerCachedFileResponse)
+		{
+			FString Filename = ServerFileTimePair.Key;
+			ConvertServerFilenameToClientFilename(Filename);
+			Filename.ToLowerInline();
+			ServerCachedFiles.Emplace(MoveTemp(Filename), ServerFileTimePair.Value);
+		}
+	}
+
+	UE_LOG(LogNetworkPlatformFile, Display, TEXT("Received '%d' cached file(s) from server"), ServerCachedFiles.Num());
 
 	bool bDeleteAllFiles = true;
 	// Check the stored cooked version
@@ -255,8 +264,8 @@ void FNetworkPlatformFile::ProcessServerCachedFilesResponse(FArrayReader& Respon
 
 	if (InnerPlatformFile->FileExists(*CookedVersionFile) == true)
 	{
-		IFileHandle* FileHandle = InnerPlatformFile->OpenRead(*CookedVersionFile);
-		if (FileHandle != NULL)
+		TUniquePtr<IFileHandle> FileHandle(InnerPlatformFile->OpenRead(*CookedVersionFile));
+		if (FileHandle.IsValid())
 		{
 			int32 StoredPackageCookedVersion;
 			int32 StoredPackageCookedLicenseeVersion;
@@ -278,8 +287,6 @@ void FNetworkPlatformFile::ProcessServerCachedFilesResponse(FArrayReader& Respon
 					}
 				}
 			}
-
-			delete FileHandle;
 		}
 	}
 	else
@@ -287,105 +294,94 @@ void FNetworkPlatformFile::ProcessServerCachedFilesResponse(FArrayReader& Respon
 		UE_LOG(LogNetworkPlatformFile, Display, TEXT("Cooked version file missing: %s\n"), *CookedVersionFile);
 	}
 
-	if (bDeleteAllFiles == true)
+	TMap<FString, FDateTime> ClientLocalFiles;
 	{
-		// Make sure the config file exists...
-		InnerPlatformFile->CreateDirectoryTree(*(FPaths::GeneratedConfigDir()));
-		// Update the cooked version file
-		IFileHandle* FileHandle = InnerPlatformFile->OpenWrite(*CookedVersionFile);
-		if (FileHandle != NULL)
+		TArray<FString> DirectoriesToSkip;
+		TArray<FString> DirectoriesToNotRecurse;
+		FLocalTimestampDirectoryVisitor Visitor(*InnerPlatformFile, DirectoriesToSkip, DirectoriesToNotRecurse, false, /* bMakeLowerCased */ true);
+
+		for (TArray<FString>::TConstIterator RootPathIt(ServerRootContentDirectories); RootPathIt; ++RootPathIt)
 		{
-			FileHandle->Write((const uint8*)&ServerPackageVersion, sizeof(int32));
-			FileHandle->Write((const uint8*)&ServerPackageLicenseeVersion, sizeof(int32));
-			delete FileHandle;
+			const FString& ContentFolder = *RootPathIt;
+			InnerPlatformFile->IterateDirectory(*ContentFolder, Visitor);
 		}
+
+		ClientLocalFiles = MoveTemp(Visitor.FileTimes);
 	}
 
-	// list of directories to skip
-	TArray<FString> DirectoriesToSkip;
-	TArray<FString> DirectoriesToNotRecurse;
-	// use the timestamp grabbing visitor to get all the content times
-	FLocalTimestampDirectoryVisitor Visitor(*InnerPlatformFile, DirectoriesToSkip, DirectoriesToNotRecurse, false);
+	UE_LOG(LogNetworkPlatformFile, Display, TEXT("Found '%d' locally cached file(s)"), ClientLocalFiles.Num());
 
-	/*TArray<FString> RootContentPaths;
-	FPackageName::QueryRootContentPaths(RootContentPaths); */
-	for (TArray<FString>::TConstIterator RootPathIt(ServerRootContentDirectories); RootPathIt; ++RootPathIt)
+	if (FParse::Param(FCommandLine::Get(), TEXT("DeleteLocalCache")))
 	{
-		/*const FString& RootPath = *RootPathIt;
-		const FString& ContentFolder = FPackageName::LongPackageNameToFilename(RootPath);*/
-		const FString& ContentFolder = *RootPathIt;
-		InnerPlatformFile->IterateDirectory(*ContentFolder, Visitor);
+		bDeleteAllFiles = true;
 	}
 
-	// delete out of date files using the server cached files
-	for (TMap<FString, FDateTime>::TIterator It(ServerCachedFiles); It; ++It)
+	if (bDeleteAllFiles)
 	{
-		bool bDeleteFile = bDeleteAllFiles;
-		FString ServerFile = It.Key();
+		UE_LOG(LogNetworkPlatformFile, Display, TEXT("Deleting all '%d' cached file(s) due to missing cooked version file..."), ClientLocalFiles.Num());
 
-		// Convert the filename to the client version
-		ConvertServerFilenameToClientFilename(ServerFile);
+		// Make sure the config file exists...
+		InnerPlatformFile->CreateDirectoryTree(*(FPaths::GeneratedConfigDir()));
+		// Update the cooked version file
+		TUniquePtr<IFileHandle> FileHandle(InnerPlatformFile->OpenWrite(*CookedVersionFile));
+		if (FileHandle.IsValid())
+		{
+			FileHandle->Write((const uint8*)&ServerPackageVersion, sizeof(int32));
+			FileHandle->Write((const uint8*)&ServerPackageLicenseeVersion, sizeof(int32));
+		}
+		else
+		{
+			UE_LOG(LogNetworkPlatformFile, Warning, TEXT("Failed to write cooked version file '%s'"), *CookedVersionFile);
+		}
+	}
 
-		// Set it in the visitor file times list
-		// If there is any pathing difference (relative path, or whatever) between the server's filelist and the results
-		// of platform directory iteration then this will Add a new entry rather than override the existing one.  This causes local file deletes
-		// and longer loads as we will never see the benefits of local device caching.
-		Visitor.FileTimes.Add(ServerFile, FDateTime::MinValue());
+	int32 NumDeleted = 0;
+	for (const auto& ClientFileTimePair : ClientLocalFiles)
+	{
+		const FString& Filename = ClientFileTimePair.Key;
+		bool bDeleteFile = bDeleteAllFiles;
 
-		if (bDeleteFile == false)
+		if (bDeleteFile)
 		{
-			// Check the time stamps...
-			// get local time
-			FDateTime LocalTime = InnerPlatformFile->GetTimeStamp(*ServerFile);
-			// If local time == MinValue than the file does not exist in the cache.
-			if (LocalTime != FDateTime::MinValue())
+			UE_LOG(LogNetworkPlatformFile, Display, TEXT("Deleting cached file '%s'"), *Filename);
+		}
+		else
+		{
+			if (const FDateTime* ServerTime = ServerCachedFiles.Find(Filename))
 			{
-				FDateTime ServerTime = It.Value();
-				// delete if out of date
-				// We will use 1.0 second as the tolerance to cover any platform differences in resolution
-				FTimespan TimeDiff = LocalTime - ServerTime;
+				FDateTime ClientTime = ClientFileTimePair.Value;
+				FTimespan TimeDiff = *ServerTime - ClientTime;
 				double TimeDiffInSeconds = TimeDiff.GetTotalSeconds();
 				bDeleteFile = (TimeDiffInSeconds > 1.0) || (TimeDiffInSeconds < -1.0);
-				if (bDeleteFile == true)
+
+				if (bDeleteFile)
 				{
-					if (InnerPlatformFile->FileExists(*ServerFile) == true)
-					{
-						UE_LOG(LogNetworkPlatformFile, Display, TEXT("Deleting cached file: TimeDiff %5.3f, %s"), TimeDiffInSeconds, *It.Key());
-					}
-					else
-					{
-						// It's a directory
-						bDeleteFile = false;
-					}
+					UE_LOG(LogNetworkPlatformFile, Display, TEXT("Deleting outdated cached file '%s' with time difference '%.3f's"), *Filename, TimeDiffInSeconds);
 				}
 				else
 				{
-					UE_LOG(LogNetworkPlatformFile, Display, TEXT("Keeping cached file: %s, TimeDiff worked out ok"), *ServerFile);
+					UE_LOG(LogNetworkPlatformFile, Display, TEXT("Keeping cached file '%s' with time difference '%.3f's"), *Filename, TimeDiffInSeconds);
 				}
 			}
+			else
+			{
+				UE_LOG(LogNetworkPlatformFile, Display, TEXT("Deleting cached file '%s' missing on server"), *Filename);
+				bDeleteFile = true;
+			}
 		}
-		if (bDeleteFile == true)
-		{
-			UE_LOG(LogNetworkPlatformFile, Display, TEXT("Deleting cached file: %s"), *ServerFile);
-			InnerPlatformFile->DeleteFile(*ServerFile);
-		}
-	}
 
-	// Any content files we have locally that were not cached, delete them
-	for (TMap<FString, FDateTime>::TIterator It(Visitor.FileTimes); It; ++It)
-	{
-		if ( FCString::Stricmp( *FPaths::GetExtension( It.Key() ), TEXT("pak")) == 0 )
+		if (bDeleteFile)
 		{
 			// ignore pak files they won't be mounted anyway 
-			continue;
-		}
-		if (It.Value() != FDateTime::MinValue())
-		{
-			// This was *not* found in the server file list... delete it
-			UE_LOG(LogNetworkPlatformFile, Display, TEXT("Deleting cached file: %s"), *It.Key());
-			InnerPlatformFile->DeleteFile(*It.Key());
+			if (FCString::Stricmp(*FPaths::GetExtension(Filename), TEXT("pak")) != 0)
+			{
+				InnerPlatformFile->DeleteFile(*Filename);
+				++NumDeleted;
+			}
 		}
 	}
+
+	UE_LOG(LogNetworkPlatformFile, Display, TEXT("Deleted '%d' of '%d' cached file(s)"), NumDeleted, ClientLocalFiles.Num());
 }
 
 FNetworkPlatformFile::~FNetworkPlatformFile()
diff --git a/Engine/Source/Runtime/SandboxFile/Private/IPlatformFileSandboxWrapper.cpp b/Engine/Source/Runtime/SandboxFile/Private/IPlatformFileSandboxWrapper.cpp
--- a/Engine/Source/Runtime/SandboxFile/Private/IPlatformFileSandboxWrapper.cpp
+++ b/Engine/Source/Runtime/SandboxFile/Private/IPlatformFileSandboxWrapper.cpp
@@ -488,6 +488,18 @@ void FSandboxPlatformFile::AddExclusion(const TCHAR* Wildcard, bool bIsDirectory
 	}
 }
 
+void FSandboxPlatformFile::RemoveExclusion(const TCHAR* Wildcard, bool bIsDirectory)
+{
+	if (bIsDirectory)
+	{
+		DirectoryExclusionWildcards.Remove(FString(Wildcard));
+	}
+	else
+	{
+		FileExclusionWildcards.Remove(FString(Wildcard));
+	}
+}
+
 void FSandboxPlatformFile::SetSandboxOnly(bool bInSandboxOnly)
 {
 	bSandboxOnly = bInSandboxOnly;
diff --git a/Engine/Source/Runtime/SandboxFile/Public/IPlatformFileSandboxWrapper.h b/Engine/Source/Runtime/SandboxFile/Public/IPlatformFileSandboxWrapper.h
--- a/Engine/Source/Runtime/SandboxFile/Public/IPlatformFileSandboxWrapper.h
+++ b/Engine/Source/Runtime/SandboxFile/Public/IPlatformFileSandboxWrapper.h
@@ -168,6 +168,7 @@ public:
 	 * @Caution, these have a performance cost
 	*/
 	void AddExclusion(const TCHAR* Wildcard, bool bIsDirectory = false);
+	void RemoveExclusion(const TCHAR* Wildcard, bool bIsDirectory = false);
 
 	/** Whether access is restricted the the sandbox or not. */
 	void SetSandboxOnly(bool bInSandboxOnly);
