Speculative fix for clang static analysis use-after-free warning

diff --git a/Engine/Source/Runtime/NetworkFile/Private/NetworkPlatformFile.cpp b/Engine/Source/Runtime/NetworkFile/Private/NetworkPlatformFile.cpp
--- a/Engine/Source/Runtime/NetworkFile/Private/NetworkPlatformFile.cpp
+++ b/Engine/Source/Runtime/NetworkFile/Private/NetworkPlatformFile.cpp
@@ -919,18 +919,28 @@ public:
 	FString							Filename;
 	/** An archive to read the file contents from */
 	FArrayReader* FileArchive;
+	TUniquePtr<FArrayReader> FileArchiveOwner;
+
 	/** timestamp for the file **/
 	FDateTime ServerTimeStamp;
 	IPlatformFile& InnerPlatformFile;
 	FScopedEvent* Event;
 
-	/** Constructor
-	*/
-	FAsyncNetworkWriteWorker(const TCHAR* InFilename, FArrayReader* InArchive, FDateTime InServerTimeStamp, IPlatformFile* InInnerPlatformFile, FScopedEvent* InEvent)
+	FAsyncNetworkWriteWorker(const TCHAR* InFilename, FArrayReader& InArchive, FDateTime InServerTimeStamp, IPlatformFile& InInnerPlatformFile)
+		: Filename(InFilename)
+		, FileArchive(&InArchive)
+		, ServerTimeStamp(InServerTimeStamp)
+		, InnerPlatformFile(InInnerPlatformFile)
+		, Event(nullptr)
+	{
+	}
+
+	FAsyncNetworkWriteWorker(const TCHAR* InFilename, TUniquePtr<FArrayReader> InArchive, FDateTime InServerTimeStamp, IPlatformFile& InInnerPlatformFile, FScopedEvent* InEvent)
 		: Filename(InFilename)
-		, FileArchive(InArchive)
+		, FileArchive(InArchive.Get())
+		, FileArchiveOwner(MoveTemp(InArchive))
 		, ServerTimeStamp(InServerTimeStamp)
-		, InnerPlatformFile(*InInnerPlatformFile)
+		, InnerPlatformFile(InInnerPlatformFile)
 		, Event(InEvent)
 	{
 	}
@@ -978,7 +988,8 @@ public:
 					// delete async write archives
 					if (Event)
 					{
-						delete FileArchive;
+						FileArchive = nullptr;
+						FileArchiveOwner.Reset();
 					}
 
 					if (InnerPlatformFile.FileSize(*TempFilename) != FileSize)
@@ -1020,15 +1031,14 @@ public:
 /**
  * Write a file async or sync, with the data coming from a FArrayReader
  */
-void SyncWriteFile(FArrayReader* Archive, const FString& Filename, FDateTime ServerTimeStamp, IPlatformFile& InnerPlatformFile)
+static void SyncWriteFile(FArrayReader* Archive, const FString& Filename, FDateTime ServerTimeStamp, IPlatformFile& InnerPlatformFile)
 {
-	FScopedEvent* NullEvent = NULL;
-	(new FAutoDeleteAsyncTask<FAsyncNetworkWriteWorker>(*Filename, Archive, ServerTimeStamp, &InnerPlatformFile, NullEvent))->StartSynchronousTask();
+	(new FAutoDeleteAsyncTask<FAsyncNetworkWriteWorker>(*Filename, *Archive, ServerTimeStamp, InnerPlatformFile))->StartSynchronousTask();
 }
 
-void AsyncWriteFile(FArrayReader* Archive, const FString& Filename, FDateTime ServerTimeStamp, IPlatformFile& InnerPlatformFile, FScopedEvent* Event = NULL)
+static void AsyncWriteFile(TUniquePtr<FArrayReader> Archive, const FString& Filename, FDateTime ServerTimeStamp, IPlatformFile& InnerPlatformFile, FScopedEvent* Event)
 {
-	(new FAutoDeleteAsyncTask<FAsyncNetworkWriteWorker>(*Filename, Archive, ServerTimeStamp, &InnerPlatformFile, Event))->StartBackgroundTask();
+	(new FAutoDeleteAsyncTask<FAsyncNetworkWriteWorker>(*Filename, MoveTemp(Archive), ServerTimeStamp, InnerPlatformFile, Event))->StartBackgroundTask();
 }
 
 void AsyncReadUnsolicitedFiles(int32 InNumUnsolictedFiles, FNetworkPlatformFile& InNetworkFile, IPlatformFile& InInnerPlatformFile, FString& InServerEngineDir, FString& InServerProjectDir, FString& InServerEnginePlatformExtensionsDir, FString& InServerProjectPlatformExtensionsDir, FScopedEvent *InNetworkDoneEvent, FScopedEvent *InWritingDoneEvent)
@@ -1065,7 +1075,7 @@ void AsyncReadUnsolicitedFiles(int32 InNumUnsolictedFiles, FNetworkPlatformFile&
 			OutstandingAsyncWrites.Add( NumUnsolictedFiles );
 			for (int32 Index = 0; Index < NumUnsolictedFiles; Index++)
 			{
-				FArrayReader* UnsolictedResponse = new FArrayReader;
+				TUniquePtr<FArrayReader> UnsolictedResponse(new FArrayReader);
 				if (!NetworkFile.ReceiveResponse(*UnsolictedResponse))
 				{
 					UE_LOG(LogNetworkPlatformFile, Fatal, TEXT("Receive failure!"));
@@ -1082,7 +1092,7 @@ void AsyncReadUnsolicitedFiles(int32 InNumUnsolictedFiles, FNetworkPlatformFile&
 					*UnsolictedResponse << UnsolictedServerTimeStamp;
 
 					// write the file by pulling out of the FArrayReader
-					AsyncWriteFile(UnsolictedResponse, UnsolictedReplyFile, UnsolictedServerTimeStamp, InnerPlatformFile, WritingDoneEvent);
+					AsyncWriteFile(MoveTemp(UnsolictedResponse), UnsolictedReplyFile, UnsolictedServerTimeStamp, InnerPlatformFile, WritingDoneEvent);
 				}
 			}
 			NetworkDoneEvent->Trigger();
