Reverting some changs to COTF platform extension handling, and using old methodology, which will not work for Programs (they wouldn't have worked before anyway)

diff --git a/Engine/Source/Runtime/NetworkFile/Private/NetworkPlatformFile.cpp b/Engine/Source/Runtime/NetworkFile/Private/NetworkPlatformFile.cpp
--- a/Engine/Source/Runtime/NetworkFile/Private/NetworkPlatformFile.cpp
+++ b/Engine/Source/Runtime/NetworkFile/Private/NetworkPlatformFile.cpp
@@ -32,6 +32,18 @@ FString FNetworkPlatformFile::BulkFileExtension = TEXT(".ubulk");
 FString FNetworkPlatformFile::ExpFileExtension = TEXT(".uexp");
 FString FNetworkPlatformFile::FontFileExtension = TEXT(".ufont");
 
+// These are marked unsafe because they do not work with Programs. However, COTF is unlikely to be used with Programs
+// These are also temporary until some issues can be debugged
+static FString UnsafeEnginePlatformExtensionDir()
+{
+	return FPaths::EnginePlatformExtensionDir(TEXT("")).TrimChar('/');
+}
+
+static FString UnsafeProjectPlatformExtensionDir()
+{
+	return FPaths::ProjectPlatformExtensionDir(TEXT("")).TrimChar('/');
+}
+
 FNetworkPlatformFile::FNetworkPlatformFile()
 	: bHasLoadedDDCDirectories(false)
 	, InnerPlatformFile(NULL)
@@ -794,19 +806,17 @@ void FNetworkPlatformFile::FillGetFileList(FNetworkFileArchive& Payload)
 		GameName = FPaths::GetProjectFilePath();
 	}
 
-	FString IniPlatformName = FPlatformProperties::IniPlatformName();
 	FString EngineRelPath = FPaths::EngineDir();
 	FString EngineRelPluginPath = FPaths::EnginePluginsDir();
 	FString GameRelPath = FPaths::ProjectDir();
 	FString GameRelPluginPath = FPaths::ProjectPluginsDir();
-	FString EnginePlatformExtensionsDir = FPaths::EnginePlatformExtensionDir(*IniPlatformName);
-	FString ProjectPlatformExtensionsDir = FPaths::ProjectPlatformExtensionDir(*IniPlatformName);
+	FString EnginePlatformExtensionsDir = UnsafeEnginePlatformExtensionDir();
+	FString ProjectPlatformExtensionsDir = UnsafeProjectPlatformExtensionDir();
 
 	TArray<FString> Directories;
 	Directories.Add(EngineRelPath);
 	Directories.Add(GameRelPath);
 
-	Payload << IniPlatformName;
 	Payload << TargetPlatformNames;
 	Payload << GameName;
 	Payload << EngineRelPath;
@@ -843,15 +853,14 @@ void FNetworkPlatformFile::ProcessServerInitialResponse(FArrayReader& InResponse
 	InResponse << ServerEnginePlatformExtensionsDir;
 	InResponse << ServerProjectPlatformExtensionsDir;
 
-	FString IniPlatformName = FPlatformProperties::IniPlatformName();
 	UE_LOG(LogNetworkPlatformFile, Display, TEXT("    Server EngineDir      = %s"), *ServerEngineDir);
 	UE_LOG(LogNetworkPlatformFile, Display, TEXT("     Local EngineDir      = %s"), *FPaths::EngineDir());
 	UE_LOG(LogNetworkPlatformFile, Display, TEXT("    Server ProjectDir     = %s"), *ServerProjectDir);
 	UE_LOG(LogNetworkPlatformFile, Display, TEXT("     Local ProjectDir     = %s"), *FPaths::ProjectDir());
 	UE_LOG(LogNetworkPlatformFile, Display, TEXT("    Server EnginePlatformExtDir = %s"), *ServerEnginePlatformExtensionsDir);
-	UE_LOG(LogNetworkPlatformFile, Display, TEXT("     Local EnginePlatformExtDir = %s"), *FPaths::EnginePlatformExtensionDir(*IniPlatformName));
+	UE_LOG(LogNetworkPlatformFile, Display, TEXT("     Local EnginePlatformExtDir = %s"), *UnsafeEnginePlatformExtensionDir());
 	UE_LOG(LogNetworkPlatformFile, Display, TEXT("    Server ProjectPlatformExtDir = %s"), *ServerProjectPlatformExtensionsDir);
-	UE_LOG(LogNetworkPlatformFile, Display, TEXT("     Local ProjectPlatformExtDir = %s"), *FPaths::ProjectPlatformExtensionDir(*IniPlatformName));
+	UE_LOG(LogNetworkPlatformFile, Display, TEXT("     Local ProjectPlatformExtDir = %s"), *UnsafeProjectPlatformExtensionDir());
 
 	// Receive a list of files and their timestamps.
 	TMap<FString, FDateTime> ServerFileMap;
@@ -1482,13 +1491,11 @@ void FNetworkPlatformFile::ConvertServerFilenameToClientFilename(FString& Filena
 	}
 	else if (FilenameToConvert.StartsWith(InServerEnginePlatformExtensionsDir))
 	{
-		FString IniPlatformName(FPlatformProperties::IniPlatformName());
-		FilenameToConvert = FilenameToConvert.Replace(*InServerEnginePlatformExtensionsDir, *(FPaths::EnginePlatformExtensionDir(*IniPlatformName)));
+		FilenameToConvert = FilenameToConvert.Replace(*InServerEnginePlatformExtensionsDir, *(UnsafeEnginePlatformExtensionDir()));
 	}
 	else if (FilenameToConvert.StartsWith(InServerProjectPlatformExtensionsDir))
 	{
-		FString IniPlatformName(FPlatformProperties::IniPlatformName());
-		FilenameToConvert = FilenameToConvert.Replace(*InServerProjectPlatformExtensionsDir, *(FPaths::ProjectPlatformExtensionDir(*IniPlatformName)));
+		FilenameToConvert = FilenameToConvert.Replace(*InServerProjectPlatformExtensionsDir, *(UnsafeProjectPlatformExtensionDir()));
 	}
 }
 
diff --git a/Engine/Source/Runtime/NetworkFileSystem/Private/NetworkFileServerConnection.cpp b/Engine/Source/Runtime/NetworkFileSystem/Private/NetworkFileServerConnection.cpp
--- a/Engine/Source/Runtime/NetworkFileSystem/Private/NetworkFileServerConnection.cpp
+++ b/Engine/Source/Runtime/NetworkFileSystem/Private/NetworkFileServerConnection.cpp
@@ -63,6 +63,18 @@ struct FSandboxOnlyScope
 	FSandboxPlatformFile& Sandbox;
 };
 
+// These are marked unsafe because they do not work with Programs. However, COTF is unlikely to be used with Programs
+// These are also temporary until some issues can be debugged
+static FString UnsafeEnginePlatformExtensionDir()
+{
+	return FPaths::EnginePlatformExtensionDir(TEXT("")).TrimChar('/');
+}
+
+static FString UnsafeProjectPlatformExtensionDir()
+{
+	return FPaths::ProjectPlatformExtensionDir(TEXT("")).TrimChar('/');
+}
+
 /* FNetworkFileServerClientConnection structors
  *****************************************************************************/
 
@@ -91,6 +103,8 @@ FNetworkFileServerClientConnection::FNetworkFileServerClientConnection(const FNe
 
 	LocalEngineDir = FPaths::EngineDir();
 	LocalProjectDir = FPaths::ProjectDir();
+	LocalEnginePlatformExtensionsDir = UnsafeEnginePlatformExtensionDir();
+	LocalProjectPlatformExtensionsDir = UnsafeProjectPlatformExtensionDir();
 
 	if (FPaths::IsProjectFilePathSet())
 	{
@@ -100,6 +114,8 @@ FNetworkFileServerClientConnection::FNetworkFileServerClientConnection(const FNe
 
 	LocalEngineDirAbs = MakeAbsoluteNormalizedDir(LocalEngineDir);
 	LocalProjectDirAbs = MakeAbsoluteNormalizedDir(LocalProjectDir);
+	LocalEnginePlatformExtensionsDirAbs = MakeAbsoluteNormalizedDir(LocalEnginePlatformExtensionsDir);
+	LocalProjectPlatformExtensionsDirAbs = MakeAbsoluteNormalizedDir(LocalEnginePlatformExtensionsDir);
 }
 
 
@@ -153,11 +169,11 @@ void FNetworkFileServerClientConnection::ConvertClientFilenameToServerFilename(F
 		// We do *not* want to replace the directory in that case.
 		return;
 	}
-	if (TrySubstituteDirectory(FilenameToConvert, FPaths::EnginePlatformExtensionDir(*ConnectedIniPlatformName), ConnectedEnginePlatformExtensionsDir))
+	if (TrySubstituteDirectory(FilenameToConvert, UnsafeEnginePlatformExtensionDir(), ConnectedEnginePlatformExtensionsDir))
 	{
 		return;
 	}
-	if (TrySubstituteDirectory(FilenameToConvert, FPaths::ProjectPlatformExtensionDir(*ConnectedIniPlatformName), ConnectedProjectPlatformExtensionsDir))
+	if (TrySubstituteDirectory(FilenameToConvert, UnsafeProjectPlatformExtensionDir(), ConnectedProjectPlatformExtensionsDir))
 	{
 		return;
 	}
@@ -871,7 +887,6 @@ bool FNetworkFileServerClientConnection::ProcessGetFileList( FArchive& In, FArch
 	FString ClientVersionInfo;
 	FString TargetAddress;
 
-	In << ConnectedIniPlatformName;
 	In << TargetPlatformNames;
 	In << GameName;
 	In << EngineRelativePath;
@@ -958,12 +973,6 @@ bool FNetworkFileServerClientConnection::ProcessGetFileList( FArchive& In, FArch
 	ConnectedProjectDir = GameRelativePath;
 	ConnectedEnginePlatformExtensionsDir = EnginePlatformExtensionsRelativePath;
 	ConnectedProjectPlatformExtensionsDir = ProjectPlatformExtensionsRelativePath;
-	
-	// now that we have a connected platform for this connection, we can query the platform extension directory locations
-	LocalEnginePlatformExtensionsDir = FPaths::EnginePlatformExtensionDir(*ConnectedIniPlatformName);
-	LocalProjectPlatformExtensionsDir = FPaths::ProjectPlatformExtensionDir(*ConnectedIniPlatformName);
-	LocalEnginePlatformExtensionsDirAbs = MakeAbsoluteNormalizedDir(LocalEnginePlatformExtensionsDir);
-	LocalProjectPlatformExtensionsDirAbs = MakeAbsoluteNormalizedDir(LocalEnginePlatformExtensionsDir);
 
 	UE_LOG(LogFileServer, Display, TEXT("    Connected EngineDir      = %s"), *ConnectedEngineDir);
 	UE_LOG(LogFileServer, Display, TEXT("        Local EngineDir      = %s"), *LocalEngineDir);
diff --git a/Engine/Source/Runtime/NetworkFileSystem/Private/NetworkFileServerConnection.h b/Engine/Source/Runtime/NetworkFileSystem/Private/NetworkFileServerConnection.h
--- a/Engine/Source/Runtime/NetworkFileSystem/Private/NetworkFileServerConnection.h
+++ b/Engine/Source/Runtime/NetworkFileSystem/Private/NetworkFileServerConnection.h
@@ -188,10 +188,7 @@ protected:
 	
 private:
 
-	// Hold the ini name of the currently connected platform.
-	FString ConnectedIniPlatformName;
-
-	// Hold the name of the currently connected target platform.
+	// Hold the name of the currently connected platform.
 	FString ConnectedPlatformName;
 
 	// Hold the ip address of the currently connected platform.
